Update  Viewports to use the new InputKey and InputAxis functions because they will be deprecated in 5.1

diff --git a/Engine/Plugins/Experimental/ActorPalette/Source/ActorPalette/Private/ActorPaletteViewportClient.cpp b/Engine/Plugins/Experimental/ActorPalette/Source/ActorPalette/Private/ActorPaletteViewportClient.cpp
--- a/Engine/Plugins/Experimental/ActorPalette/Source/ActorPalette/Private/ActorPaletteViewportClient.cpp
+++ b/Engine/Plugins/Experimental/ActorPalette/Source/ActorPalette/Private/ActorPaletteViewportClient.cpp
@@ -45,13 +45,13 @@ FActorPaletteViewportClient::FActorPaletteViewportClient(int32 InTabIndex)
 	EngineShowFlags.Grid = false;
 }
 
-bool FActorPaletteViewportClient::InputKey(FViewport* InViewport, int32 ControllerId, FKey Key, EInputEvent Event, float AmountDepressed, bool bGamepad)
+bool FActorPaletteViewportClient::InputKey(const FInputKeyEventArgs& InEventArgs)
 {
-	if (Key == EKeys::LeftMouseButton)
+	if (InEventArgs.Key == EKeys::LeftMouseButton)
 	{
 		if (FSlateApplication::Get().IsDragDropping())
 		{
-// 			if (Event == IE_Released)
+// 			if (InEventArgs.Event == IE_Released)
 // 			{
 // 				FSlateApplication::Get().CancelDragDrop();
 // 			}
@@ -59,8 +59,8 @@ bool FActorPaletteViewportClient::InputKey(FViewport* InViewport, int32 Controll
 			return true;
 		}
 
-		const int32	HitX = InViewport->GetMouseX();
-		const int32	HitY = InViewport->GetMouseY();
+		const int32	HitX = InEventArgs.Viewport->GetMouseX();
+		const int32	HitY = InEventArgs.Viewport->GetMouseY();
 
 		// Calc the raw delta from the mouse to detect if there was any movement
 		FVector RawMouseDelta = MouseDeltaTracker->GetRawDelta();
@@ -76,7 +76,7 @@ bool FActorPaletteViewportClient::InputKey(FViewport* InViewport, int32 Controll
 		}
 		else
 		{
-			HHitProxy* HitProxy = InViewport->GetHitProxy(HitX, HitY);
+			HHitProxy* HitProxy = InEventArgs.Viewport->GetHitProxy(HitX, HitY);
 
 
 			if (HActor* ActorProxy = HitProxyCast<HActor>(HitProxy))
@@ -100,7 +100,7 @@ bool FActorPaletteViewportClient::InputKey(FViewport* InViewport, int32 Controll
 						TSharedPtr<FAssetDragDropOp> DragDropOperation = FAssetDragDropOp::New(FAssetData(Assets[0], true));
 
 						TSet<FKey> PressedMouseButtons;
-						PressedMouseButtons.Add(Key);
+						PressedMouseButtons.Add(InEventArgs.Key);
 
 						FModifierKeysState ModifierKeyState;
 
@@ -131,7 +131,9 @@ bool FActorPaletteViewportClient::InputKey(FViewport* InViewport, int32 Controll
 	}
 	else
 	{
-		return FEditorViewportClient::InputKey(InViewport, ControllerId, (Key == EKeys::RightMouseButton) ? EKeys::LeftMouseButton : Key, Event, AmountDepressed, bGamepad);
+		FInputKeyEventArgs Args = InEventArgs;
+		Args.Key = (InEventArgs.Key == EKeys::RightMouseButton) ? EKeys::LeftMouseButton : InEventArgs.Key;
+		return FEditorViewportClient::InputKey(Args);
 	}
 }
 
diff --git a/Engine/Plugins/Experimental/ActorPalette/Source/ActorPalette/Private/ActorPaletteViewportClient.h b/Engine/Plugins/Experimental/ActorPalette/Source/ActorPalette/Private/ActorPaletteViewportClient.h
--- a/Engine/Plugins/Experimental/ActorPalette/Source/ActorPalette/Private/ActorPaletteViewportClient.h
+++ b/Engine/Plugins/Experimental/ActorPalette/Source/ActorPalette/Private/ActorPaletteViewportClient.h
@@ -20,7 +20,7 @@ public:
 	FActorPaletteViewportClient(int32 InTabIndex);
 
 	// FViewportClient interface
-	virtual bool InputKey(FViewport* Viewport, int32 ControllerId, FKey Key, EInputEvent Event, float AmountDepressed, bool bGamepad) override;
+	virtual bool InputKey(const FInputKeyEventArgs& InEventArgs) override;
 	// End of FViewportClient interface
 
 	// FEditorViewportClient interface
diff --git a/Engine/Source/Editor/Kismet/Private/SCSEditorViewportClient.cpp b/Engine/Source/Editor/Kismet/Private/SCSEditorViewportClient.cpp
--- a/Engine/Source/Editor/Kismet/Private/SCSEditorViewportClient.cpp
+++ b/Engine/Source/Editor/Kismet/Private/SCSEditorViewportClient.cpp
@@ -288,13 +288,13 @@ void FSCSEditorViewportClient::DrawCanvas( FViewport& InViewport, FSceneView& Vi
 	}
 }
 
-bool FSCSEditorViewportClient::InputKey(FViewport* InViewport, int32 ControllerId, FKey Key, EInputEvent Event, float AmountDepressed, bool bGamepad)
+bool FSCSEditorViewportClient::InputKey(const FInputKeyEventArgs& EventArgs)
 {
-	bool bHandled = GUnrealEd->ComponentVisManager.HandleInputKey(this, InViewport, Key, Event);;
+	bool bHandled = GUnrealEd->ComponentVisManager.HandleInputKey(this, EventArgs.Viewport, EventArgs.Key, EventArgs.Event);
 
-	if( !bHandled )
+	if(!bHandled)
 	{
-		bHandled = FEditorViewportClient::InputKey(InViewport, ControllerId, Key, Event, AmountDepressed, bGamepad);
+		bHandled = FEditorViewportClient::InputKey(EventArgs);
 	}
 
 	return bHandled;
diff --git a/Engine/Source/Editor/Kismet/Private/SCSEditorViewportClient.h b/Engine/Source/Editor/Kismet/Private/SCSEditorViewportClient.h
--- a/Engine/Source/Editor/Kismet/Private/SCSEditorViewportClient.h
+++ b/Engine/Source/Editor/Kismet/Private/SCSEditorViewportClient.h
@@ -37,7 +37,7 @@ public:
 	virtual void Tick(float DeltaSeconds) override;
 	virtual void Draw(const FSceneView* View,FPrimitiveDrawInterface* PDI) override;
 	virtual void DrawCanvas( FViewport& InViewport, FSceneView& View, FCanvas& Canvas ) override;
-	virtual bool InputKey(FViewport* Viewport, int32 ControllerId, FKey Key, EInputEvent Event, float AmountDepressed = 1.f, bool bGamepad=false) override;
+	virtual bool InputKey(const FInputKeyEventArgs& EventArgs) override;
 	virtual void ProcessClick(class FSceneView& View, class HHitProxy* HitProxy, FKey Key, EInputEvent Event, uint32 HitX, uint32 HitY) override;
 	virtual bool InputWidgetDelta( FViewport* Viewport, EAxisList::Type CurrentAxis, FVector& Drag, FRotator& Rot, FVector& Scale ) override;
 	virtual void TrackingStarted( const struct FInputEventState& InInputState, bool bIsDragging, bool bNudge ) override;
