Fix C4855: implicit capture of 'this' via '[=]' is deprecated in '/std:c++20'

#rnx

diff --git a/Engine/Plugins/2D/Paper2D/Source/SmartSnapping/Private/SmartSnappingModule.cpp b/Engine/Plugins/2D/Paper2D/Source/SmartSnapping/Private/SmartSnappingModule.cpp
--- a/Engine/Plugins/2D/Paper2D/Source/SmartSnapping/Private/SmartSnappingModule.cpp
+++ b/Engine/Plugins/2D/Paper2D/Source/SmartSnapping/Private/SmartSnappingModule.cpp
@@ -47,7 +47,7 @@ public:
 			// Unregister the level editor extensions
 			{
 				FLevelEditorModule& LevelEditor = FModuleManager::GetModuleChecked<FLevelEditorModule>( TEXT("LevelEditor") );
-				LevelEditor.GetAllLevelEditorToolbarViewMenuExtenders().RemoveAll([=](const FLevelEditorModule::FLevelEditorMenuExtender& Extender) { return Extender.GetHandle() == ViewMenuExtenderHandle; });
+				LevelEditor.GetAllLevelEditorToolbarViewMenuExtenders().RemoveAll([this](const FLevelEditorModule::FLevelEditorMenuExtender& Extender) { return Extender.GetHandle() == ViewMenuExtenderHandle; });
 			}
 			
 			//
diff --git a/Engine/Plugins/Animation/ControlRig/Source/ControlRigEditor/Private/Sequencer/ControlRigParameterTrackEditor.cpp b/Engine/Plugins/Animation/ControlRig/Source/ControlRigEditor/Private/Sequencer/ControlRigParameterTrackEditor.cpp
--- a/Engine/Plugins/Animation/ControlRig/Source/ControlRigEditor/Private/Sequencer/ControlRigParameterTrackEditor.cpp
+++ b/Engine/Plugins/Animation/ControlRig/Source/ControlRigEditor/Private/Sequencer/ControlRigParameterTrackEditor.cpp
@@ -3663,7 +3663,7 @@ void FControlRigParameterTrackEditor::AddControlKeys(
 	
 	TGuardValue<bool> Guard(bIsDoingSelection, true);
 
-	auto OnKeyProperty = [=](FFrameNumber Time) -> FKeyPropertyResult
+	auto OnKeyProperty = [=, this](FFrameNumber Time) -> FKeyPropertyResult
 	{
 		FFrameNumber LocalTime = Time;
 		//for modify weights we evaluate so need to make sure we use the evaluated time
@@ -4865,7 +4865,7 @@ void FControlRigParameterSection::BuildSectionContextMenu(FMenuBuilder& MenuBuil
 						LOCTEXT("CollapseAllSections", "Collapse All Sections"),
 						LOCTEXT("CollapseAllSections_ToolTip", "Collapse all sections onto this section"),
 						FSlateIcon(),
-						FUIAction(FExecuteAction::CreateLambda([=] { CollapseAllLayers(); }))
+						FUIAction(FExecuteAction::CreateLambda([this] { CollapseAllLayers(); }))
 					);
 				}
 			}
@@ -4878,7 +4878,7 @@ void FControlRigParameterSection::BuildSectionContextMenu(FMenuBuilder& MenuBuil
 						LOCTEXT("KeyZeroValue", "Key Zero Value"),
 						LOCTEXT("KeyZeroValue_Tooltip", "Set zero key on all controls in this section"),
 						FSlateIcon(),
-						FUIAction(FExecuteAction::CreateLambda([=] { KeyZeroValue(); }))
+						FUIAction(FExecuteAction::CreateLambda([this] { KeyZeroValue(); }))
 					);
 				}
 				
@@ -4886,14 +4886,14 @@ void FControlRigParameterSection::BuildSectionContextMenu(FMenuBuilder& MenuBuil
 					LOCTEXT("KeyWeightZero", "Key Weight Zero"),
 					LOCTEXT("KeyWeightZero_Tooltip", "Key a zero value on the Weight channel"),
 					FSlateIcon(),
-					FUIAction(FExecuteAction::CreateLambda([=] { KeyWeightValue(0.0f); }))
+					FUIAction(FExecuteAction::CreateLambda([this] { KeyWeightValue(0.0f); }))
 				);
 				
 				MenuBuilder.AddMenuEntry(
 					LOCTEXT("KeyWeightOne", "Key Weight One"),
 					LOCTEXT("KeyWeightOne_Tooltip", "Key a one value on the Weight channel"),
 					FSlateIcon(),
-					FUIAction(FExecuteAction::CreateLambda([=] { KeyWeightValue(1.0f); }))
+					FUIAction(FExecuteAction::CreateLambda([this] { KeyWeightValue(1.0f); }))
 				);
 				
 			}
@@ -4905,8 +4905,8 @@ void FControlRigParameterSection::BuildSectionContextMenu(FMenuBuilder& MenuBuil
 				LOCTEXT("SetFromSelectedControls_ToolTip", "Set active channels from the current control selection"),
 				FSlateIcon(),
 				FUIAction(
-					FExecuteAction::CreateLambda([=] { ShowSelectedControlsChannels(); }),
-					FCanExecuteAction::CreateLambda([=] { return ControlRig->CurrentControlSelection().Num() > 0; } )
+					FExecuteAction::CreateLambda([this] { ShowSelectedControlsChannels(); }),
+					FCanExecuteAction::CreateLambda([ControlRig] { return ControlRig->CurrentControlSelection().Num() > 0; } )
 				)
 			);
 
@@ -4914,7 +4914,7 @@ void FControlRigParameterSection::BuildSectionContextMenu(FMenuBuilder& MenuBuil
 				LOCTEXT("ShowAllControls", "Show All Controls"),
 				LOCTEXT("ShowAllControls_ToolTip", "Set active channels from all controls"),
 				FSlateIcon(),
-				FUIAction(FExecuteAction::CreateLambda([=] { return ShowAllControlsChannels(); }))
+				FUIAction(FExecuteAction::CreateLambda([this] { return ShowAllControlsChannels(); }))
 			);
 
 			MenuBuilder.AddSubMenu(
diff --git a/Engine/Plugins/Animation/LiveLink/Source/LiveLinkEditor/Private/LiveLinkTest.cpp b/Engine/Plugins/Animation/LiveLink/Source/LiveLinkEditor/Private/LiveLinkTest.cpp
--- a/Engine/Plugins/Animation/LiveLink/Source/LiveLinkEditor/Private/LiveLinkTest.cpp
+++ b/Engine/Plugins/Animation/LiveLink/Source/LiveLinkEditor/Private/LiveLinkTest.cpp
@@ -102,7 +102,7 @@ bool FLiveLinkInterpolationTest::RunTest(const FString& Parameters)
 	BasicWorker->GenericInterpolate(BlendFactor, InterpolationOptions, FrameDataStructA, FrameDataStructB, FrameDataStructR);
 	const FLiveLinkTestFrameDataInternal& FrameR = *FrameDataStructR.Cast<FLiveLinkTestFrameDataInternal>();
 
-	auto TestStruct = [=](const FLiveLinkInnerTestInternal& FrameR, int32 OuterIndex)
+	auto TestStruct = [=, this](const FLiveLinkInnerTestInternal& FrameR, int32 OuterIndex)
 		{
 			TestEqual(TEXT("InnerSingleFloat"), FrameR.InnerSingleFloat, FloatR * -1.f - 100 * OuterIndex);
 			TestEqual(TEXT("InnerSingleInt"), FrameR.InnerSingleInt, IntR * -1 - 100 * OuterIndex);
diff --git a/Engine/Plugins/Animation/LiveLink/Source/LiveLinkSequencer/Private/SequencerRecorderSections/MovieSceneLiveLinkTrackRecorder.cpp b/Engine/Plugins/Animation/LiveLink/Source/LiveLinkSequencer/Private/SequencerRecorderSections/MovieSceneLiveLinkTrackRecorder.cpp
--- a/Engine/Plugins/Animation/LiveLink/Source/LiveLinkSequencer/Private/SequencerRecorderSections/MovieSceneLiveLinkTrackRecorder.cpp
+++ b/Engine/Plugins/Animation/LiveLink/Source/LiveLinkSequencer/Private/SequencerRecorderSections/MovieSceneLiveLinkTrackRecorder.cpp
@@ -90,7 +90,7 @@ void UMovieSceneLiveLinkTrackRecorder::CreateTracks()
 	const bool bIncludeDisabledSubjects = false;
 	const bool bIncludeVirtualSubjects = true;
 	TArray<FLiveLinkSubjectKey> EnabledSubjects = LiveLinkClient->GetSubjects(bIncludeDisabledSubjects, bIncludeVirtualSubjects);
-	const FLiveLinkSubjectKey* DesiredSubjectKey = EnabledSubjects.FindByPredicate([=](const FLiveLinkSubjectKey& InOther) { return SubjectName == InOther.SubjectName; });
+	const FLiveLinkSubjectKey* DesiredSubjectKey = EnabledSubjects.FindByPredicate([this](const FLiveLinkSubjectKey& InOther) { return SubjectName == InOther.SubjectName; });
 	if (DesiredSubjectKey == nullptr)
 	{
 		UE_LOG(LogLiveLinkSequencer, Warning, TEXT("Error: Could not create live link track. Could not find an enabled subject with subject name '%s'."), *SubjectName.ToString());
diff --git a/Engine/Plugins/Cameras/CameraShakePreviewer/Source/CameraShakePreviewer/Private/CameraShakePreviewerModule.cpp b/Engine/Plugins/Cameras/CameraShakePreviewer/Source/CameraShakePreviewer/Private/CameraShakePreviewerModule.cpp
--- a/Engine/Plugins/Cameras/CameraShakePreviewer/Source/CameraShakePreviewer/Private/CameraShakePreviewerModule.cpp
+++ b/Engine/Plugins/Cameras/CameraShakePreviewer/Source/CameraShakePreviewer/Private/CameraShakePreviewerModule.cpp
@@ -181,7 +181,7 @@ void FCameraShakePreviewerModule::UnregisterViewportOptionMenuExtender()
 	UToolMenus::UnregisterOwner(this);
 
 	FLevelEditorModule& LevelEditorModule = FModuleManager::GetModuleChecked<FLevelEditorModule>(LevelEditorModuleName);
-	LevelEditorModule.GetAllLevelViewportOptionsMenuExtenders().RemoveAll([=](const FLevelEditorModule::FLevelEditorMenuExtender& Extender) { return Extender.GetHandle() == ViewportOptionsMenuExtenderHandle; });
+	LevelEditorModule.GetAllLevelViewportOptionsMenuExtenders().RemoveAll([this](const FLevelEditorModule::FLevelEditorMenuExtender& Extender) { return Extender.GetHandle() == ViewportOptionsMenuExtenderHandle; });
 
 	if (GEditor)
 	{
diff --git a/Engine/Plugins/Compositing/OpenColorIO/Source/OpenColorIOEditor/Private/OpenColorIOColorSpaceConversionCustomization.cpp b/Engine/Plugins/Compositing/OpenColorIO/Source/OpenColorIOEditor/Private/OpenColorIOColorSpaceConversionCustomization.cpp
--- a/Engine/Plugins/Compositing/OpenColorIO/Source/OpenColorIOEditor/Private/OpenColorIOColorSpaceConversionCustomization.cpp
+++ b/Engine/Plugins/Compositing/OpenColorIO/Source/OpenColorIOEditor/Private/OpenColorIOColorSpaceConversionCustomization.cpp
@@ -45,7 +45,7 @@ void FOpenColorIOColorConversionSettingsCustomization::CustomizeHeader(TSharedRe
 					.VAlign(VAlign_Center)
 					[
 						SNew(STextBlock)
-						.Text(MakeAttributeLambda([=]
+						.Text(MakeAttributeLambda([this]
 							{
 								if (ColorSpaceConversion->IsValid())
 								{
diff --git a/Engine/Plugins/Compositing/OpenColorIO/Source/OpenColorIOEditor/Private/OpenColorIOColorSpaceCustomization.cpp b/Engine/Plugins/Compositing/OpenColorIO/Source/OpenColorIOEditor/Private/OpenColorIOColorSpaceCustomization.cpp
--- a/Engine/Plugins/Compositing/OpenColorIO/Source/OpenColorIOEditor/Private/OpenColorIOColorSpaceCustomization.cpp
+++ b/Engine/Plugins/Compositing/OpenColorIO/Source/OpenColorIOEditor/Private/OpenColorIOColorSpaceCustomization.cpp
@@ -193,7 +193,7 @@ void FOpenColorIOColorSpaceCustomization::AddMenuEntry(FMenuBuilder& InMenuBuild
 		FText::FromString(InColorSpace.ToString()),
 		FSlateIcon(),
 		FUIAction(
-			FExecuteAction::CreateLambda([=]
+			FExecuteAction::CreateLambda([this, InColorSpace]
 			{
 				if (FStructProperty* StructProperty = CastField<FStructProperty>(CachedProperty->GetProperty()))
 				{
@@ -208,7 +208,7 @@ void FOpenColorIOColorSpaceCustomization::AddMenuEntry(FMenuBuilder& InMenuBuild
 			}
 			),
 			FCanExecuteAction(),
-			FIsActionChecked::CreateLambda([=]
+			FIsActionChecked::CreateLambda([this, InColorSpace]
 			{
 				TArray<void*> RawData;
 				CachedProperty->AccessRawData(RawData);
@@ -285,7 +285,7 @@ void FOpenColorIODisplayViewCustomization::AddMenuEntry(FMenuBuilder& InMenuBuil
 		FText::FromString(InDisplayView.ToString()),
 		FSlateIcon(),
 		FUIAction(
-			FExecuteAction::CreateLambda([=]
+			FExecuteAction::CreateLambda([this, InDisplayView]
 				{
 					if (FStructProperty* StructProperty = CastField<FStructProperty>(CachedProperty->GetProperty()))
 					{
@@ -300,7 +300,7 @@ void FOpenColorIODisplayViewCustomization::AddMenuEntry(FMenuBuilder& InMenuBuil
 				}
 			),
 			FCanExecuteAction(),
-			FIsActionChecked::CreateLambda([=]
+			FIsActionChecked::CreateLambda([this, InDisplayView]
 				{
 					TArray<void*> RawData;
 					CachedProperty->AccessRawData(RawData);
diff --git a/Engine/Plugins/Compositing/OpenColorIO/Source/OpenColorIOEditor/Private/Widgets/SOpenColorIOColorSpacePicker.cpp b/Engine/Plugins/Compositing/OpenColorIO/Source/OpenColorIOEditor/Private/Widgets/SOpenColorIOColorSpacePicker.cpp
--- a/Engine/Plugins/Compositing/OpenColorIO/Source/OpenColorIOEditor/Private/Widgets/SOpenColorIOColorSpacePicker.cpp
+++ b/Engine/Plugins/Compositing/OpenColorIO/Source/OpenColorIOEditor/Private/Widgets/SOpenColorIOColorSpacePicker.cpp
@@ -38,7 +38,7 @@ void SOpenColorIOColorSpacePicker::Construct(const FArguments& InArgs)
 		.VAlign(VAlign_Center)
 		[
 			SNew(STextBlock)
-			.Text(MakeAttributeLambda([=]
+			.Text(MakeAttributeLambda([this]
 				{
 					if (DisplayViewSelection.IsValid())
 					{
@@ -158,7 +158,7 @@ TSharedRef<SWidget> SOpenColorIOColorSpacePicker::HandleColorSpaceComboButtonMen
 							SetCurrentColorSpace(ColorSpace);
 						}),
 						FCanExecuteAction(),
-						FIsActionChecked::CreateLambda([=]
+						FIsActionChecked::CreateLambda([this, ColorSpace]
 						{
 							return ColorSpaceSelection == ColorSpace;
 						})
@@ -203,7 +203,7 @@ TSharedRef<SWidget> SOpenColorIOColorSpacePicker::HandleColorSpaceComboButtonMen
 								SetCurrentDisplayView(DisplayView);
 							}),
 							FCanExecuteAction(),
-							FIsActionChecked::CreateLambda([=]
+							FIsActionChecked::CreateLambda([this, DisplayView]
 							{
 								return DisplayViewSelection == DisplayView;
 							})
diff --git a/Engine/Plugins/Editor/GameplayTagsEditor/Source/GameplayTagsEditor/Private/SGameplayTagWidget.cpp b/Engine/Plugins/Editor/GameplayTagsEditor/Source/GameplayTagsEditor/Private/SGameplayTagWidget.cpp
--- a/Engine/Plugins/Editor/GameplayTagsEditor/Source/GameplayTagsEditor/Private/SGameplayTagWidget.cpp
+++ b/Engine/Plugins/Editor/GameplayTagsEditor/Source/GameplayTagsEditor/Private/SGameplayTagWidget.cpp
@@ -407,7 +407,7 @@ bool SGameplayTagWidget::FilterChildrenCheck( TSharedPtr<FGameplayTagNode> InIte
 		return false;
 	}
 
-	auto FilterChildrenCheck_r = ([=]()
+	auto FilterChildrenCheck_r = ([this, InItem]()
 	{
 		TArray< TSharedPtr<FGameplayTagNode> > Children = InItem->GetChildTagNodes();
 		for( int32 iChild = 0; iChild < Children.Num(); ++iChild )
diff --git a/Engine/Plugins/Enterprise/DatasmithImporter/Source/DatasmithImporter/Private/DatasmithFileProducer.cpp b/Engine/Plugins/Enterprise/DatasmithImporter/Source/DatasmithImporter/Private/DatasmithFileProducer.cpp
--- a/Engine/Plugins/Enterprise/DatasmithImporter/Source/DatasmithImporter/Private/DatasmithFileProducer.cpp
+++ b/Engine/Plugins/Enterprise/DatasmithImporter/Source/DatasmithImporter/Private/DatasmithFileProducer.cpp
@@ -1618,7 +1618,7 @@ void FDatasmithFileProducerDetails::CustomizeDetails(IDetailLayoutBuilder& Detai
 		[
 			SAssignNew(IconText, STextBlock)
 				.Font(FAppStyle::Get().GetFontStyle("FontAwesome.11"))
-				.Text(MakeAttributeLambda([=]
+				.Text(MakeAttributeLambda([this]
 				{
 					return IsProducerSuperseded() ? FEditorFontGlyphs::Exclamation_Triangle : FEditorFontGlyphs::File;
 				}))
@@ -1626,7 +1626,7 @@ void FDatasmithFileProducerDetails::CustomizeDetails(IDetailLayoutBuilder& Detai
 		]
 	];
 
-	IconText->SetToolTipText(MakeAttributeLambda([=]
+	IconText->SetToolTipText(MakeAttributeLambda([this]
 	{
 		if (IsProducerSuperseded())
 		{
@@ -1841,7 +1841,7 @@ void FDatasmithDirProducerDetails::CustomizeDetails( IDetailLayoutBuilder& Detai
 		[
 			SAssignNew(IconText, STextBlock)
 			.Font(FAppStyle::Get().GetFontStyle("FontAwesome.11"))
-			.Text(MakeAttributeLambda([=]
+			.Text(MakeAttributeLambda([this]
 			{
 				return IsProducerSuperseded() ? FEditorFontGlyphs::Exclamation_Triangle : FEditorFontGlyphs::Folder;
 			}))
@@ -1849,7 +1849,7 @@ void FDatasmithDirProducerDetails::CustomizeDetails( IDetailLayoutBuilder& Detai
 		]
 	];
 
-	IconText->SetToolTipText(MakeAttributeLambda([=]
+	IconText->SetToolTipText(MakeAttributeLambda([this]
 	{
 		if (IsProducerSuperseded())
 		{
diff --git a/Engine/Plugins/Experimental/ActorPalette/Source/ActorPalette/Private/ActorPaletteViewport.cpp b/Engine/Plugins/Experimental/ActorPalette/Source/ActorPalette/Private/ActorPaletteViewport.cpp
--- a/Engine/Plugins/Experimental/ActorPalette/Source/ActorPalette/Private/ActorPaletteViewport.cpp
+++ b/Engine/Plugins/Experimental/ActorPalette/Source/ActorPalette/Private/ActorPaletteViewport.cpp
@@ -189,11 +189,11 @@ void SActorPaletteViewport::Construct(const FArguments& InArgs, TSharedPtr<class
 				SNew(SButton)
 				.Text(LOCTEXT("QuickButtonLabel_ReloadLastMap", "Open Last Level"))
 				.VAlign(VAlign_Center)
-				.Visibility_Lambda([=]()
+				.Visibility_Lambda([this]()
 				{
 					return (GetDefault<UActorPaletteSettings>()->FindLastLevelForTab(TabIndex) != INDEX_NONE) ? EVisibility::Visible : EVisibility::Collapsed;
 				})
-				.ToolTipText_Lambda([=]()
+				.ToolTipText_Lambda([this]()
 				{
 					const UActorPaletteSettings* Settings = GetDefault<UActorPaletteSettings>();
 					const int32 LastLevelForTabIndex = Settings->FindLastLevelForTab(TabIndex);
@@ -201,7 +201,7 @@ void SActorPaletteViewport::Construct(const FArguments& InArgs, TSharedPtr<class
 						FText::Format(LOCTEXT("QuickButtonTooltip_ReloadLastMap", "Use {0} as an Actor Palette"), FText::FromName(Settings->SettingsPerLevel[LastLevelForTabIndex].GetAsAssetData().PackageName)) :
 						FText::GetEmpty();
 				})
-				.OnClicked_Lambda([=]()
+				.OnClicked_Lambda([this]()
 				{
 					const UActorPaletteSettings* Settings = GetDefault<UActorPaletteSettings>();
 					int32 LastLevelIndex = Settings->FindLastLevelForTab(TabIndex);
@@ -279,13 +279,13 @@ void SActorPaletteViewport::BindCommands()
 
 	CommandList->MapAction(
 		Commands.ToggleGameView,
-		FExecuteAction::CreateLambda([=]() { TypedViewportClient->SetGameView(!TypedViewportClient->IsInGameView()); }),
+		FExecuteAction::CreateLambda([this]() { TypedViewportClient->SetGameView(!TypedViewportClient->IsInGameView()); }),
 		FCanExecuteAction(),
-		FIsActionChecked::CreateLambda([=](){ return TypedViewportClient->IsInGameView(); }));
+		FIsActionChecked::CreateLambda([this](){ return TypedViewportClient->IsInGameView(); }));
 
 	CommandList->MapAction(
 		Commands.ResetCameraView,
-		FExecuteAction::CreateLambda([=]() { TypedViewportClient->ResetCameraView(); }));
+		FExecuteAction::CreateLambda([this]() { TypedViewportClient->ResetCameraView(); }));
 }
 
 void SActorPaletteViewport::OnFocusViewportToSelection()
@@ -343,7 +343,7 @@ TSharedRef<SWidget> SActorPaletteViewport::GenerateMapMenu() const
 		InMenuBuilder.AddSubMenu(
 			LOCTEXT("OpenLevelFromPicker_Label", "Open Level..."),
 			LOCTEXT("OpenLevelFromPicker_Tooltip", "Select an existing level asset to use as an Actor Palette"),
-			FNewMenuDelegate::CreateLambda([=](FMenuBuilder& SubMenuBuilder)
+			FNewMenuDelegate::CreateLambda([this](FMenuBuilder& SubMenuBuilder)
 			{
 				FContentBrowserModule& ContentBrowserModule = FModuleManager::Get().LoadModuleChecked<FContentBrowserModule>(TEXT("ContentBrowser"));
 							
@@ -351,7 +351,7 @@ TSharedRef<SWidget> SActorPaletteViewport::GenerateMapMenu() const
 				FAssetPickerConfig Config;
 				Config.Filter.ClassPaths.Add(UWorld::StaticClass()->GetClassPathName());
 				Config.InitialAssetViewType = EAssetViewType::List;
-				Config.OnAssetSelected = FOnAssetSelected::CreateLambda([=](const FAssetData& AssetData)
+				Config.OnAssetSelected = FOnAssetSelected::CreateLambda([this](const FAssetData& AssetData)
 				{
 					TypedViewportClient->OpenWorldAsPalette(AssetData);
 					FSlateApplication::Get().DismissAllMenus();
@@ -387,7 +387,7 @@ TSharedRef<SWidget> SActorPaletteViewport::GenerateMapMenu() const
 			if (Asset.AssetClassPath == UWorld::StaticClass()->GetClassPathName())
 			{
 				FUIAction Action;
-				Action.ExecuteAction.BindLambda([=]()
+				Action.ExecuteAction.BindLambda([this, Asset]()
 				{
 					TypedViewportClient->OpenWorldAsPalette(Asset);
 				});
@@ -412,7 +412,7 @@ TSharedRef<SWidget> SActorPaletteViewport::GenerateMapMenu() const
 		InMenuBuilder.AddSubMenu(
 			LOCTEXT("OpenLevelFromRecentList_Label", "Recent Levels"),
 			LOCTEXT("OpenLevelFromRecentList_Tooltip", "Select a level recently used as an Actor Palette"),
-			FNewMenuDelegate::CreateLambda([=](FMenuBuilder& SubMenuBuilder)
+			FNewMenuDelegate::CreateLambda([this, Settings](FMenuBuilder& SubMenuBuilder)
 			{
 				bool bAddedAnyRecent = false;
 				for (const FString& RecentEntry : Settings->RecentlyUsedList)
@@ -424,7 +424,7 @@ TSharedRef<SWidget> SActorPaletteViewport::GenerateMapMenu() const
 						bAddedAnyRecent = true;
 
 						FUIAction Action;
-						Action.ExecuteAction.BindLambda([=]()
+						Action.ExecuteAction.BindLambda([this, Asset]()
 						{
 							TypedViewportClient->OpenWorldAsPalette(Asset);
 						});
@@ -451,7 +451,7 @@ TSharedRef<SWidget> SActorPaletteViewport::GenerateMapMenu() const
 		InMenuBuilder.AddSubMenu(
 			LOCTEXT("OpenLevelFromFavoriteList_Label", "Favorite Levels"),
 			LOCTEXT("OpenLevelFromFavoriteList_Tooltip", "Select a favorite level as an Actor Palette"),
-			FNewMenuDelegate::CreateLambda([=](FMenuBuilder& SubMenuBuilder)
+			FNewMenuDelegate::CreateLambda([this, Settings](FMenuBuilder& SubMenuBuilder)
 			{
 				if (TypedViewportClient->GetCurrentWorldAssetData().IsValid())
 				{
diff --git a/Engine/Plugins/Experimental/Animation/SkeletalMeshModelingTools/Source/SkeletalMeshModelingTools/Private/SkeletalMeshModelingToolsModule.cpp b/Engine/Plugins/Experimental/Animation/SkeletalMeshModelingTools/Source/SkeletalMeshModelingTools/Private/SkeletalMeshModelingToolsModule.cpp
--- a/Engine/Plugins/Experimental/Animation/SkeletalMeshModelingTools/Source/SkeletalMeshModelingTools/Private/SkeletalMeshModelingToolsModule.cpp
+++ b/Engine/Plugins/Experimental/Animation/SkeletalMeshModelingTools/Source/SkeletalMeshModelingTools/Private/SkeletalMeshModelingToolsModule.cpp
@@ -62,11 +62,11 @@ void FSkeletalMeshModelingToolsModule::ShutdownModule()
 	{
 		// un-register toolbar extender delegates
 		TArray<ISkeletalMeshEditorModule::FSkeletalMeshEditorToolbarExtender>& Extenders = SkelMeshEditorModule->GetAllSkeletalMeshEditorToolbarExtenders();
-		Extenders.RemoveAll([=](const ISkeletalMeshEditorModule::FSkeletalMeshEditorToolbarExtender& InDelegate) { return InDelegate.GetHandle() == SkelMeshEditorExtenderHandle; });
+		Extenders.RemoveAll([this](const ISkeletalMeshEditorModule::FSkeletalMeshEditorToolbarExtender& InDelegate) { return InDelegate.GetHandle() == SkelMeshEditorExtenderHandle; });
 
 		// un-register post-init delegates
 		TArray<ISkeletalMeshEditorModule::FOnSkeletalMeshEditorInitialized>& PostInitDelegates = SkelMeshEditorModule->GetPostEditorInitDelegates();
-		PostInitDelegates.RemoveAll([=](const ISkeletalMeshEditorModule::FOnSkeletalMeshEditorInitialized& InDelegate) { return InDelegate.GetHandle() == SkelMeshEditorPostInitHandle; });
+		PostInitDelegates.RemoveAll([this](const ISkeletalMeshEditorModule::FOnSkeletalMeshEditorInitialized& InDelegate) { return InDelegate.GetHandle() == SkelMeshEditorPostInitHandle; });
 	}
 
 	UToolMenus::UnregisterOwner(this);
diff --git a/Engine/Plugins/Experimental/ChaosCaching/Source/ChaosCachingEditor/Private/Sequencer/ChaosCacheTrackEditor.cpp b/Engine/Plugins/Experimental/ChaosCaching/Source/ChaosCachingEditor/Private/Sequencer/ChaosCacheTrackEditor.cpp
--- a/Engine/Plugins/Experimental/ChaosCaching/Source/ChaosCachingEditor/Private/Sequencer/ChaosCacheTrackEditor.cpp
+++ b/Engine/Plugins/Experimental/ChaosCaching/Source/ChaosCachingEditor/Private/Sequencer/ChaosCacheTrackEditor.cpp
@@ -357,7 +357,7 @@ TSharedPtr<SWidget> FChaosCacheTrackEditor::BuildOutlinerEditWidget(const FGuid&
 	{
 		TWeakPtr<ISequencer> WeakSequencer = GetSequencer();
 
-		auto SubMenuCallback = [=]() -> TSharedRef<SWidget>
+		auto SubMenuCallback = [this, ObjectBinding, Track]() -> TSharedRef<SWidget>
 		{
 			FMenuBuilder MenuBuilder(true, nullptr);
 
diff --git a/Engine/Plugins/Experimental/ChaosEditor/Source/FractureEditor/Private/FractureEditorModeToolkit.cpp b/Engine/Plugins/Experimental/ChaosEditor/Source/FractureEditor/Private/FractureEditorModeToolkit.cpp
--- a/Engine/Plugins/Experimental/ChaosEditor/Source/FractureEditor/Private/FractureEditorModeToolkit.cpp
+++ b/Engine/Plugins/Experimental/ChaosEditor/Source/FractureEditor/Private/FractureEditorModeToolkit.cpp
@@ -1104,19 +1104,19 @@ void FFractureEditorModeToolkit::BindCommands()
 
 	ToolkitCommands->MapAction(
 		Commands.ExplodeMore,
-		FExecuteAction::CreateLambda([=]() { this->OnSetExplodedViewValue( FMath::Min(1.0, this->GetExplodedViewValue() + .1) ); } ),
+		FExecuteAction::CreateLambda([this]() { this->OnSetExplodedViewValue( FMath::Min(1.0, this->GetExplodedViewValue() + .1) ); } ),
 		EUIActionRepeatMode::RepeatEnabled
 	);
 
 	ToolkitCommands->MapAction(
 		Commands.ExplodeLess,
-		FExecuteAction::CreateLambda([=]() { this->OnSetExplodedViewValue( FMath::Max(0.0, this->GetExplodedViewValue() - .1) ); } ),
+		FExecuteAction::CreateLambda([this]() { this->OnSetExplodedViewValue( FMath::Max(0.0, this->GetExplodedViewValue() - .1) ); } ),
 		EUIActionRepeatMode::RepeatEnabled
 	);
 
 	ToolkitCommands->MapAction(
 		Commands.CancelTool,
-		FExecuteAction::CreateLambda([=]()
+		FExecuteAction::CreateLambda([this]()
 		{
 			if (GetActiveTool())
 			{
diff --git a/Engine/Plugins/Experimental/GameFeatures/Source/GameFeaturesEditor/Private/GameFeatureDataDetailsCustomization.cpp b/Engine/Plugins/Experimental/GameFeatures/Source/GameFeaturesEditor/Private/GameFeatureDataDetailsCustomization.cpp
--- a/Engine/Plugins/Experimental/GameFeatures/Source/GameFeaturesEditor/Private/GameFeatureDataDetailsCustomization.cpp
+++ b/Engine/Plugins/Experimental/GameFeatures/Source/GameFeaturesEditor/Private/GameFeatureDataDetailsCustomization.cpp
@@ -132,7 +132,7 @@ void FGameFeatureDataDetailsCustomization::CustomizeDetails(IDetailLayoutBuilder
 				.Padding(0.0f, 4.0f, 0.0f, 0.0f)
 				[
 					SNew(SHorizontalBox)
-					.Visibility_Lambda([=]() { return (GetCurrentState() == EGameFeaturePluginState::Active) ? EVisibility::Visible : EVisibility::Collapsed; })
+					.Visibility_Lambda([this]() { return (GetCurrentState() == EGameFeaturePluginState::Active) ? EVisibility::Visible : EVisibility::Collapsed; })
 					+SHorizontalBox::Slot()
 					.AutoWidth()
 					.Padding(Padding)
diff --git a/Engine/Plugins/Experimental/GeometryCollectionPlugin/Source/GeometryCollectionSequencer/Private/GeometryCollectionTrackEditor.cpp b/Engine/Plugins/Experimental/GeometryCollectionPlugin/Source/GeometryCollectionSequencer/Private/GeometryCollectionTrackEditor.cpp
--- a/Engine/Plugins/Experimental/GeometryCollectionPlugin/Source/GeometryCollectionSequencer/Private/GeometryCollectionTrackEditor.cpp
+++ b/Engine/Plugins/Experimental/GeometryCollectionPlugin/Source/GeometryCollectionSequencer/Private/GeometryCollectionTrackEditor.cpp
@@ -381,7 +381,7 @@ TSharedPtr<SWidget> FGeometryCollectionTrackEditor::BuildOutlinerEditWidget(cons
 	{
 		TWeakPtr<ISequencer> WeakSequencer = GetSequencer();
 
-		auto SubMenuCallback = [=]() -> TSharedRef<SWidget>
+		auto SubMenuCallback = [this, ObjectBinding, Track, GeomMeshComp]() -> TSharedRef<SWidget>
 		{
 			FMenuBuilder MenuBuilder(true, nullptr);
 
diff --git a/Engine/Plugins/Experimental/Mutable/Source/CustomizableObjectEditor/Private/MuCOE/SCustomizableObjectEditorViewportToolBar.cpp b/Engine/Plugins/Experimental/Mutable/Source/CustomizableObjectEditor/Private/MuCOE/SCustomizableObjectEditorViewportToolBar.cpp
--- a/Engine/Plugins/Experimental/Mutable/Source/CustomizableObjectEditor/Private/MuCOE/SCustomizableObjectEditorViewportToolBar.cpp
+++ b/Engine/Plugins/Experimental/Mutable/Source/CustomizableObjectEditor/Private/MuCOE/SCustomizableObjectEditorViewportToolBar.cpp
@@ -1090,14 +1090,14 @@ TSharedRef<SWidget> SCustomizableObjectEditorViewportToolBar::GenerateViewportOp
 						SNew(SBox).WidthOverride(100.0f)
 						[
 							SNew(SSlider)
-							.Value_Lambda([=]() {return Viewport.Pin().Get()->GetViewportCameraSpeed(); })
+							.Value_Lambda([this]() {return Viewport.Pin().Get()->GetViewportCameraSpeed(); })
 							.MinValue(1).MaxValue(4)
-							.OnValueChanged_Lambda([=](int32 Value) { Viewport.Pin().Get()->SetViewportCameraSpeed(Value); })
+							.OnValueChanged_Lambda([this](int32 Value) { Viewport.Pin().Get()->SetViewportCameraSpeed(Value); })
 						]
 					]
 					+ SHorizontalBox::Slot().AutoWidth()
 					[
-						SNew(STextBlock).Text_Lambda([=]() { return FText::AsNumber(Viewport.Pin().Get()->GetViewportCameraSpeed()); })
+						SNew(STextBlock).Text_Lambda([this]() { return FText::AsNumber(Viewport.Pin().Get()->GetViewportCameraSpeed()); })
 					]
 				];
 				
diff --git a/Engine/Plugins/Experimental/Mutable/Source/CustomizableObjectPopulationEditor/Private/MuCOPE/CustomizableObjectPopulationClassEditor.cpp b/Engine/Plugins/Experimental/Mutable/Source/CustomizableObjectPopulationEditor/Private/MuCOPE/CustomizableObjectPopulationClassEditor.cpp
--- a/Engine/Plugins/Experimental/Mutable/Source/CustomizableObjectPopulationEditor/Private/MuCOPE/CustomizableObjectPopulationClassEditor.cpp
+++ b/Engine/Plugins/Experimental/Mutable/Source/CustomizableObjectPopulationEditor/Private/MuCOPE/CustomizableObjectPopulationClassEditor.cpp
@@ -714,16 +714,16 @@ TSharedRef<SWidget> FCustomizableObjectPopulationClassEditor::GenerateTestPopula
 		ViewportColumnsEntry = SNew(SNumericEntryBox<int32>).AllowSpin(true)
 		.MinValue(1).MaxValue(100)
 		.MinSliderValue(1).MaxSliderValue(100)
-		.Value_Lambda([=]()->int32 {return ViewportColumns; })
-		.OnValueChanged_Lambda([=](int32 NewValue) {ViewportColumns = NewValue; });
+		.Value_Lambda([this]()->int32 {return ViewportColumns; })
+		.OnValueChanged_Lambda([this](int32 NewValue) {ViewportColumns = NewValue; });
 
 		MenuBuilder.AddWidget(ViewportColumnsEntry.ToSharedRef(), LOCTEXT("PopulationClassColumns", "Population Columns:"));
 
 		InstanceSeparationEntry = SNew(SNumericEntryBox<int32>).AllowSpin(true)
 		.MinValue(1).MaxValue(1000)
 		.MinSliderValue(1).MaxSliderValue(1000)
-		.Value_Lambda([=]()->int32 {return InstanceSeparation; })
-		.OnValueChanged_Lambda([=](int32 NewValue) {InstanceSeparation = NewValue; });
+		.Value_Lambda([this]()->int32 {return InstanceSeparation; })
+		.OnValueChanged_Lambda([this](int32 NewValue) {InstanceSeparation = NewValue; });
 
 		MenuBuilder.AddWidget(InstanceSeparationEntry.ToSharedRef(), LOCTEXT("PopulationClassSeparation", "Instance separation:"));
 	}
diff --git a/Engine/Plugins/Experimental/Mutable/Source/CustomizableObjectPopulationEditor/Private/MuCOPE/CustomizableObjectPopulationEditor.cpp b/Engine/Plugins/Experimental/Mutable/Source/CustomizableObjectPopulationEditor/Private/MuCOPE/CustomizableObjectPopulationEditor.cpp
--- a/Engine/Plugins/Experimental/Mutable/Source/CustomizableObjectPopulationEditor/Private/MuCOPE/CustomizableObjectPopulationEditor.cpp
+++ b/Engine/Plugins/Experimental/Mutable/Source/CustomizableObjectPopulationEditor/Private/MuCOPE/CustomizableObjectPopulationEditor.cpp
@@ -359,16 +359,16 @@ TSharedRef<SWidget> FCustomizableObjectPopulationEditor::GenerateTestPopulationM
 		ViewportColumnsEntry = SNew(SNumericEntryBox<int32>).AllowSpin(true)
 		.MinValue(1).MaxValue(100)
 		.MinSliderValue(1).MaxSliderValue(100)
-		.Value_Lambda([=]()->int32 {return ViewportColumns; })
-		.OnValueChanged_Lambda([=](int32 NewValue) {ViewportColumns = NewValue; });
+		.Value_Lambda([this]()->int32 {return ViewportColumns; })
+		.OnValueChanged_Lambda([this](int32 NewValue) {ViewportColumns = NewValue; });
 
 		MenuBuilder.AddWidget(ViewportColumnsEntry.ToSharedRef(), LOCTEXT("PopulationClassColumns", "Population Columns:"));
 
 		InstanceSeparationEntry = SNew(SNumericEntryBox<int32>).AllowSpin(true)
 		.MinValue(1).MaxValue(1000)
 		.MinSliderValue(1).MaxSliderValue(1000)
-		.Value_Lambda([=]()->int32 {return InstanceSeparation; })
-		.OnValueChanged_Lambda([=](int32 NewValue) {InstanceSeparation = NewValue; });
+		.Value_Lambda([this]()->int32 {return InstanceSeparation; })
+		.OnValueChanged_Lambda([this](int32 NewValue) {InstanceSeparation = NewValue; });
 
 		MenuBuilder.AddWidget(InstanceSeparationEntry.ToSharedRef(), LOCTEXT("PopulationClassSeparation", "Instance separation:"));
 	}
diff --git a/Engine/Plugins/Experimental/NNE/Source/NNEEditorTools/Private/NNEEditorToolsModelDataEditorToolkit.cpp b/Engine/Plugins/Experimental/NNE/Source/NNEEditorTools/Private/NNEEditorToolsModelDataEditorToolkit.cpp
--- a/Engine/Plugins/Experimental/NNE/Source/NNEEditorTools/Private/NNEEditorToolsModelDataEditorToolkit.cpp
+++ b/Engine/Plugins/Experimental/NNE/Source/NNEEditorTools/Private/NNEEditorToolsModelDataEditorToolkit.cpp
@@ -34,7 +34,7 @@ namespace UE::NNEEditorTools::Private
 
 		WorkspaceMenuCategory = InTabManager->AddLocalWorkspaceMenuCategory(INVTEXT("NNE Model Data Editor"));
 
-		InTabManager->RegisterTabSpawner("NNEModelDataTab", FOnSpawnTab::CreateLambda([=](const FSpawnTabArgs&)
+		InTabManager->RegisterTabSpawner("NNEModelDataTab", FOnSpawnTab::CreateLambda([this](const FSpawnTabArgs&)
 			{
 				return SNew(SDockTab)
 				[
diff --git a/Engine/Plugins/FX/Niagara/Source/NiagaraEditor/Private/Customizations/NiagaraDebugHUDCustomization.cpp b/Engine/Plugins/FX/Niagara/Source/NiagaraEditor/Private/Customizations/NiagaraDebugHUDCustomization.cpp
--- a/Engine/Plugins/FX/Niagara/Source/NiagaraEditor/Private/Customizations/NiagaraDebugHUDCustomization.cpp
+++ b/Engine/Plugins/FX/Niagara/Source/NiagaraEditor/Private/Customizations/NiagaraDebugHUDCustomization.cpp
@@ -716,8 +716,8 @@ void FNiagaraDebugHUDSettingsDetailsCustomization::CustomizeDetails(class IDetai
 				.AutoWidth()
 				[
 					SNew(SCheckBox)
-					.IsChecked_Lambda([=]() -> ECheckBoxState { return GetEditBool() ? ECheckBoxState::Checked : ECheckBoxState::Unchecked; } )
-					.OnCheckStateChanged_Lambda([=](ECheckBoxState NewState) { GetEditBool() = NewState == ECheckBoxState::Checked; WeakSettings.Get()->NotifyPropertyChanged(); })
+					.IsChecked_Lambda([GetEditBool]() -> ECheckBoxState { return GetEditBool() ? ECheckBoxState::Checked : ECheckBoxState::Unchecked; } )
+					.OnCheckStateChanged_Lambda([this, GetEditBool](ECheckBoxState NewState) { GetEditBool() = NewState == ECheckBoxState::Checked; WeakSettings.Get()->NotifyPropertyChanged(); })
 				]
 				+ SHorizontalBox::Slot()
 				[
diff --git a/Engine/Plugins/FX/Niagara/Source/NiagaraEditor/Private/Customizations/UserParameters/NiagaraUserParametersBuilderBase.cpp b/Engine/Plugins/FX/Niagara/Source/NiagaraEditor/Private/Customizations/UserParameters/NiagaraUserParametersBuilderBase.cpp
--- a/Engine/Plugins/FX/Niagara/Source/NiagaraEditor/Private/Customizations/UserParameters/NiagaraUserParametersBuilderBase.cpp
+++ b/Engine/Plugins/FX/Niagara/Source/NiagaraEditor/Private/Customizations/UserParameters/NiagaraUserParametersBuilderBase.cpp
@@ -335,7 +335,7 @@ void FNiagaraUserParameterNodeBuilder::AddCustomMenuActionsForParameter(FDetailW
 	if(UserParameter.IsDataInterface())
 	{
 		FUIAction CopyAction;
-		CopyAction.ExecuteAction = FExecuteAction::CreateLambda([=]()
+		CopyAction.ExecuteAction = FExecuteAction::CreateLambda([this, UserParameter]()
 		{
 			FNiagaraVariant CurrentParameterValue = GetCurrentParameterValue(UserParameter);
 			if(CurrentParameterValue.IsValid() == false || CurrentParameterValue.GetDataInterface() == nullptr)
@@ -352,7 +352,7 @@ void FNiagaraUserParameterNodeBuilder::AddCustomMenuActionsForParameter(FDetailW
 		WidgetRow.CopyAction(CopyAction);
 
 		FUIAction PasteAction;
-		PasteAction.ExecuteAction = FExecuteAction::CreateLambda([=]()
+		PasteAction.ExecuteAction = FExecuteAction::CreateLambda([this, UserParameter]()
 		{
 			FNiagaraVariant CurrentParameterValue = GetCurrentParameterValue(UserParameter);
 			if(CurrentParameterValue.IsValid() == false || CurrentParameterValue.GetDataInterface() == nullptr)
@@ -376,7 +376,7 @@ void FNiagaraUserParameterNodeBuilder::AddCustomMenuActionsForParameter(FDetailW
 	else if(UserParameter.IsUObject())
 	{
 		FUIAction CopyAction;
-		CopyAction.ExecuteAction = FExecuteAction::CreateLambda([=]()
+		CopyAction.ExecuteAction = FExecuteAction::CreateLambda([this, UserParameter]()
 		{
 			FNiagaraVariant CurrentParameterValue = GetCurrentParameterValue(UserParameter);
 			if(CurrentParameterValue.IsValid() == false)
@@ -389,7 +389,7 @@ void FNiagaraUserParameterNodeBuilder::AddCustomMenuActionsForParameter(FDetailW
 		WidgetRow.CopyAction(CopyAction);
 
 		FUIAction PasteAction;
-		PasteAction.ExecuteAction = FExecuteAction::CreateLambda([=]()
+		PasteAction.ExecuteAction = FExecuteAction::CreateLambda([this, UserParameter]()
 		{
 			FString PastedString;
 			FPropertyEditorClipboard::ClipboardPaste(PastedString);
diff --git a/Engine/Plugins/FX/Niagara/Source/NiagaraEditor/Private/NiagaraNodeStaticSwitch.cpp b/Engine/Plugins/FX/Niagara/Source/NiagaraEditor/Private/NiagaraNodeStaticSwitch.cpp
--- a/Engine/Plugins/FX/Niagara/Source/NiagaraEditor/Private/NiagaraNodeStaticSwitch.cpp
+++ b/Engine/Plugins/FX/Niagara/Source/NiagaraEditor/Private/NiagaraNodeStaticSwitch.cpp
@@ -1200,10 +1200,10 @@ void UNiagaraNodeStaticSwitch::GetNodeContextMenuActions(UToolMenu* Menu, UGraph
 			FText::Format(LOCTEXT("OpenEnumLabel", "Open {0}"), FText::FromString(SwitchTypeData.Enum->GetName())),
 			LOCTEXT("OpenEnumTooltip", "Opens up the enum asset."),
 			FSlateIcon(),
-			FUIAction(FExecuteAction::CreateLambda([=]()
+			FUIAction(FExecuteAction::CreateLambda([this]()
 			{
 				GEditor->GetEditorSubsystem<UAssetEditorSubsystem>()->OpenEditorForAsset(SwitchTypeData.Enum);
-			}),	FCanExecuteAction::CreateLambda([=]()
+			}),	FCanExecuteAction::CreateLambda([this]()
 			{
 				return Cast<UUserDefinedEnum>(SwitchTypeData.Enum) != nullptr;
 			})));
@@ -1213,7 +1213,7 @@ void UNiagaraNodeStaticSwitch::GetNodeContextMenuActions(UToolMenu* Menu, UGraph
 			FText::Format(LOCTEXT("AddEnumToTypesLabel", "Add enum {0} to user added enums"), FText::FromString(SwitchTypeData.Enum->GetName())),
 			LOCTEXT("AddEnumToTypesTooltip", "Adds the enum to the list of user added enums so it can be used as a parameter."), 
 			FSlateIcon(), 
-			FUIAction(FExecuteAction::CreateLambda([=]()
+			FUIAction(FExecuteAction::CreateLambda([this]()
 			{
 				GetMutableDefault<UNiagaraSettings>()->AddEnumParameterType(SwitchTypeData.Enum);
 			}), FCanExecuteAction::CreateLambda([this]()
diff --git a/Engine/Plugins/FX/Niagara/Source/NiagaraEditor/Private/Toolkits/NiagaraSystemToolkit.cpp b/Engine/Plugins/FX/Niagara/Source/NiagaraEditor/Private/Toolkits/NiagaraSystemToolkit.cpp
--- a/Engine/Plugins/FX/Niagara/Source/NiagaraEditor/Private/Toolkits/NiagaraSystemToolkit.cpp
+++ b/Engine/Plugins/FX/Niagara/Source/NiagaraEditor/Private/Toolkits/NiagaraSystemToolkit.cpp
@@ -524,7 +524,7 @@ void FNiagaraSystemToolkit::SetupCommands()
 
 	GetToolkitCommands()->MapAction(
 		FNiagaraEditorCommands::Get().IsolateSelectedEmitters,
-		FExecuteAction::CreateLambda([=]()
+		FExecuteAction::CreateLambda([this]()
 		{
 			SystemViewModel->IsolateSelectedEmitters();
 		}),
@@ -532,7 +532,7 @@ void FNiagaraSystemToolkit::SetupCommands()
 
 	GetToolkitCommands()->MapAction(
 		FNiagaraEditorCommands::Get().DisableSelectedEmitters,
-		FExecuteAction::CreateLambda([=]()
+		FExecuteAction::CreateLambda([this]()
 		{
 			SystemViewModel->DisableSelectedEmitters();
 		}),
diff --git a/Engine/Plugins/FX/Niagara/Source/NiagaraEditor/Private/Toolkits/SystemToolkitModes/NiagaraSystemToolkitModeBase.cpp b/Engine/Plugins/FX/Niagara/Source/NiagaraEditor/Private/Toolkits/SystemToolkitModes/NiagaraSystemToolkitModeBase.cpp
--- a/Engine/Plugins/FX/Niagara/Source/NiagaraEditor/Private/Toolkits/SystemToolkitModes/NiagaraSystemToolkitModeBase.cpp
+++ b/Engine/Plugins/FX/Niagara/Source/NiagaraEditor/Private/Toolkits/SystemToolkitModes/NiagaraSystemToolkitModeBase.cpp
@@ -640,7 +640,7 @@ TSharedRef<SDockTab> FNiagaraSystemToolkitModeBase::SpawnTab_Details(const FSpaw
 			]
 		];
 
-	SDockTab::FOnTabClosedCallback TabClosedCallback = SDockTab::FOnTabClosedCallback::CreateLambda([=](TSharedRef<SDockTab> DockTab)
+	SDockTab::FOnTabClosedCallback TabClosedCallback = SDockTab::FOnTabClosedCallback::CreateLambda([this](TSharedRef<SDockTab> DockTab)
 	{
 		SystemToolkit.Pin()->GetSystemViewModel()->GetSelectionViewModel()->GetSelectionStackViewModel()->ResetSearchText();
 	});
diff --git a/Engine/Plugins/FX/Niagara/Source/NiagaraEditor/Private/ViewModels/Stack/NiagaraStackEmitterSettingsGroup.cpp b/Engine/Plugins/FX/Niagara/Source/NiagaraEditor/Private/ViewModels/Stack/NiagaraStackEmitterSettingsGroup.cpp
--- a/Engine/Plugins/FX/Niagara/Source/NiagaraEditor/Private/ViewModels/Stack/NiagaraStackEmitterSettingsGroup.cpp
+++ b/Engine/Plugins/FX/Niagara/Source/NiagaraEditor/Private/ViewModels/Stack/NiagaraStackEmitterSettingsGroup.cpp
@@ -131,7 +131,7 @@ void UNiagaraStackEmitterPropertiesItem::RefreshChildrenInternal(const TArray<UN
 
 UNiagaraStackEntry::FStackIssueFixDelegate UNiagaraStackEmitterPropertiesItem::GetUpgradeVersionFix()
 {
-	return FStackIssueFixDelegate::CreateLambda([=]()
+	return FStackIssueFixDelegate::CreateLambda([this]()
 	{
 		FGuid NewVersion = GetEmitterViewModel()->GetParentEmitter().Emitter->GetExposedVersion().VersionGuid;
 		FNiagaraEditorUtilities::SwitchParentEmitterVersion(GetEmitterViewModel().ToSharedRef(), GetSystemViewModel(), NewVersion);
diff --git a/Engine/Plugins/FX/Niagara/Source/NiagaraEditor/Private/ViewModels/Stack/NiagaraStackEntry.cpp b/Engine/Plugins/FX/Niagara/Source/NiagaraEditor/Private/ViewModels/Stack/NiagaraStackEntry.cpp
--- a/Engine/Plugins/FX/Niagara/Source/NiagaraEditor/Private/ViewModels/Stack/NiagaraStackEntry.cpp
+++ b/Engine/Plugins/FX/Niagara/Source/NiagaraEditor/Private/ViewModels/Stack/NiagaraStackEntry.cpp
@@ -775,7 +775,7 @@ void UNiagaraStackEntry::RefreshChildren()
 	}
 	
 	// Stack issues refresh
-	NewStackIssues.RemoveAll([=](const FStackIssue& Issue) { return Issue.GetCanBeDismissed() && GetStackEditorData().GetDismissedStackIssueIds().Contains(Issue.GetUniqueIdentifier()); }); 
+	NewStackIssues.RemoveAll([this](const FStackIssue& Issue) { return Issue.GetCanBeDismissed() && GetStackEditorData().GetDismissedStackIssueIds().Contains(Issue.GetUniqueIdentifier()); }); 
 
 	StackIssues.Empty();
 	StackIssues.Append(NewStackIssues);
diff --git a/Engine/Plugins/FX/Niagara/Source/NiagaraEditor/Private/ViewModels/Stack/NiagaraStackEventScriptItemGroup.cpp b/Engine/Plugins/FX/Niagara/Source/NiagaraEditor/Private/ViewModels/Stack/NiagaraStackEventScriptItemGroup.cpp
--- a/Engine/Plugins/FX/Niagara/Source/NiagaraEditor/Private/ViewModels/Stack/NiagaraStackEventScriptItemGroup.cpp
+++ b/Engine/Plugins/FX/Niagara/Source/NiagaraEditor/Private/ViewModels/Stack/NiagaraStackEventScriptItemGroup.cpp
@@ -271,7 +271,7 @@ void UNiagaraStackEventScriptItemGroup::RefreshChildrenInternal(const TArray<UNi
 	FVersionedNiagaraEmitterData* EmitterData = GetEmitterViewModel()->GetEmitter().GetEmitterData();
 
 	const FNiagaraEventScriptProperties* EventScriptProperties = EmitterData->GetEventHandlers().FindByPredicate(
-		[=](const FNiagaraEventScriptProperties& InEventScriptProperties) { return InEventScriptProperties.Script->GetUsageId() == GetScriptUsageId(); });
+		[this](const FNiagaraEventScriptProperties& InEventScriptProperties) { return InEventScriptProperties.Script->GetUsageId() == GetScriptUsageId(); });
 
 	if (EventScriptProperties != nullptr)
 	{
diff --git a/Engine/Plugins/FX/Niagara/Source/NiagaraEditor/Private/ViewModels/Stack/NiagaraStackFunctionInput.cpp b/Engine/Plugins/FX/Niagara/Source/NiagaraEditor/Private/ViewModels/Stack/NiagaraStackFunctionInput.cpp
--- a/Engine/Plugins/FX/Niagara/Source/NiagaraEditor/Private/ViewModels/Stack/NiagaraStackFunctionInput.cpp
+++ b/Engine/Plugins/FX/Niagara/Source/NiagaraEditor/Private/ViewModels/Stack/NiagaraStackFunctionInput.cpp
@@ -627,7 +627,7 @@ UNiagaraStackEntry::FStackIssueFixDelegate UNiagaraStackFunctionInput::GetUpgrad
 	{
 		return FStackIssueFixDelegate();
 	}
-	return FStackIssueFixDelegate::CreateLambda([=]()
+	return FStackIssueFixDelegate::CreateLambda([this]()
 	{
 		FScopedTransaction ScopedTransaction(LOCTEXT("UpgradeVersionFix", "Change dynamic input version"));
 		FNiagaraScriptVersionUpgradeContext UpgradeContext;
@@ -765,7 +765,7 @@ void UNiagaraStackFunctionInput::RefreshChildrenInternal(const TArray<UNiagaraSt
 		if (ScriptData != nullptr)
 		{
 			UNiagaraStackFunctionInputCollection* DynamicInputEntry = FindCurrentChildOfTypeByPredicate<UNiagaraStackFunctionInputCollection>(CurrentChildren,
-				[=](UNiagaraStackFunctionInputCollection* CurrentFunctionInputEntry)
+				[this](UNiagaraStackFunctionInputCollection* CurrentFunctionInputEntry)
 			{
 				return CurrentFunctionInputEntry->GetInputFunctionCallNode() == InputValues.DynamicNode.Get() &&
 					CurrentFunctionInputEntry->GetModuleNode() == OwningModuleNode.Get();
@@ -916,7 +916,7 @@ void UNiagaraStackFunctionInput::RefreshChildrenInternal(const TArray<UNiagaraSt
 	if (InputValues.Mode == EValueMode::Data && InputValues.DataObject.IsValid())
 	{
 		UNiagaraStackObject* ValueObjectEntry = FindCurrentChildOfTypeByPredicate<UNiagaraStackObject>(CurrentChildren,
-			[=](UNiagaraStackObject* CurrentObjectEntry) { return CurrentObjectEntry->GetObject() == InputValues.DataObject.Get(); });
+			[this](UNiagaraStackObject* CurrentObjectEntry) { return CurrentObjectEntry->GetObject() == InputValues.DataObject.Get(); });
 
 		if(ValueObjectEntry == nullptr)
 		{
@@ -1633,7 +1633,7 @@ void UNiagaraStackFunctionInput::GetAvailableParameterHandles(TArray<FNiagaraPar
 			FNiagaraStackGraphUtilities::GetStackNodeGroups(*OutputNode, StackGroups);
 
 			int32 CurrentModuleIndex = OutputNode == CurrentOutputNode
-				? StackGroups.IndexOfByPredicate([=](const FNiagaraStackGraphUtilities::FStackNodeGroup Group) { return Group.EndNode == OwningModuleNode; })
+				? StackGroups.IndexOfByPredicate([this](const FNiagaraStackGraphUtilities::FStackNodeGroup Group) { return Group.EndNode == OwningModuleNode; })
 				: INDEX_NONE;
 
 			int32 MaxGroupIndex = CurrentModuleIndex != INDEX_NONE ? CurrentModuleIndex : StackGroups.Num() - 1;
diff --git a/Engine/Plugins/FX/Niagara/Source/NiagaraEditor/Private/ViewModels/Stack/NiagaraStackModuleItem.cpp b/Engine/Plugins/FX/Niagara/Source/NiagaraEditor/Private/ViewModels/Stack/NiagaraStackModuleItem.cpp
--- a/Engine/Plugins/FX/Niagara/Source/NiagaraEditor/Private/ViewModels/Stack/NiagaraStackModuleItem.cpp
+++ b/Engine/Plugins/FX/Niagara/Source/NiagaraEditor/Private/ViewModels/Stack/NiagaraStackModuleItem.cpp
@@ -860,7 +860,7 @@ UNiagaraStackEntry::FStackIssueFixDelegate UNiagaraStackModuleItem::GetUpgradeVe
 	{
 		return FStackIssueFixDelegate();
 	}
-	return FStackIssueFixDelegate::CreateLambda([=]()
+	return FStackIssueFixDelegate::CreateLambda([this]()
     {
         FScopedTransaction ScopedTransaction(LOCTEXT("UpgradeVersionFix", "Change module version"));
         FNiagaraScriptVersionUpgradeContext UpgradeContext;
diff --git a/Engine/Plugins/FX/Niagara/Source/NiagaraEditor/Private/ViewModels/Stack/NiagaraStackScriptItemGroup.cpp b/Engine/Plugins/FX/Niagara/Source/NiagaraEditor/Private/ViewModels/Stack/NiagaraStackScriptItemGroup.cpp
--- a/Engine/Plugins/FX/Niagara/Source/NiagaraEditor/Private/ViewModels/Stack/NiagaraStackScriptItemGroup.cpp
+++ b/Engine/Plugins/FX/Niagara/Source/NiagaraEditor/Private/ViewModels/Stack/NiagaraStackScriptItemGroup.cpp
@@ -592,7 +592,7 @@ void UNiagaraStackScriptItemGroup::RefreshIssues(TArray<FStackIssue>& NewIssues)
 		FText FixDescription = LOCTEXT("FixStackGraph", "Fix invalid stack graph");
 		FStackIssueFix ResetStackFix(
 			FixDescription,
-			FStackIssueFixDelegate::CreateLambda([=]()
+			FStackIssueFixDelegate::CreateLambda([this, FixDescription, Graph]()
 		{
 			FScopedTransaction ScopedTransaction(FixDescription);
 			FNiagaraStackGraphUtilities::ResetGraphForOutput(*Graph, ScriptUsage, ScriptUsageId);
diff --git a/Engine/Plugins/FX/Niagara/Source/NiagaraEditor/Private/Widgets/NiagaraScratchScriptEditor.cpp b/Engine/Plugins/FX/Niagara/Source/NiagaraEditor/Private/Widgets/NiagaraScratchScriptEditor.cpp
--- a/Engine/Plugins/FX/Niagara/Source/NiagaraEditor/Private/Widgets/NiagaraScratchScriptEditor.cpp
+++ b/Engine/Plugins/FX/Niagara/Source/NiagaraEditor/Private/Widgets/NiagaraScratchScriptEditor.cpp
@@ -211,7 +211,7 @@ TSharedRef<SWidget> FNiagaraGraphEditorSummoner::CreateTabBodyForObject(const FW
 
 	// We need to register the tab being closed as we need to clear out the active selection logic when that happens.
 	// There may be better places to put this, but this works well in practice.
-	SDockTab::FOnTabClosedCallback TabClosedCallback = SDockTab::FOnTabClosedCallback::CreateLambda([=](TSharedRef<SDockTab> DockTab)
+	SDockTab::FOnTabClosedCallback TabClosedCallback = SDockTab::FOnTabClosedCallback::CreateLambda([this](TSharedRef<SDockTab> DockTab)
 		{
 			TSharedPtr<FNiagaraSystemToolkit> Toolkit = EditorPtr.Pin();
 			if (Toolkit.IsValid())
diff --git a/Engine/Plugins/FX/Niagara/Source/NiagaraEditor/Private/Widgets/SNiagaraBakerWidget.cpp b/Engine/Plugins/FX/Niagara/Source/NiagaraEditor/Private/Widgets/SNiagaraBakerWidget.cpp
--- a/Engine/Plugins/FX/Niagara/Source/NiagaraEditor/Private/Widgets/SNiagaraBakerWidget.cpp
+++ b/Engine/Plugins/FX/Niagara/Source/NiagaraEditor/Private/Widgets/SNiagaraBakerWidget.cpp
@@ -225,7 +225,7 @@ void SNiagaraBakerWidget::Construct(const FArguments& InArgs)
 	// Baker Toolbar
 	FSlimHorizontalToolBarBuilder BakerToolbarBuilder(MakeShareable(new FUICommandList), FMultiBoxCustomization::None);
 	BakerToolbarBuilder.AddToolBarButton(
-		FUIAction(FExecuteAction::CreateLambda([=]() { OnCapture(); })),
+		FUIAction(FExecuteAction::CreateLambda([this]() { OnCapture(); })),
 		NAME_None,
 		FText::GetEmpty(),	//LOCTEXT("Bake", "Bake"),
 		LOCTEXT("BakeTooltip", "Run the bake process"),
diff --git a/Engine/Plugins/FX/Niagara/Source/NiagaraEditor/Private/Widgets/SNiagaraFilterBox.cpp b/Engine/Plugins/FX/Niagara/Source/NiagaraEditor/Private/Widgets/SNiagaraFilterBox.cpp
--- a/Engine/Plugins/FX/Niagara/Source/NiagaraEditor/Private/Widgets/SNiagaraFilterBox.cpp
+++ b/Engine/Plugins/FX/Niagara/Source/NiagaraEditor/Private/Widgets/SNiagaraFilterBox.cpp
@@ -37,7 +37,7 @@ void SNiagaraSourceFilterCheckBox::Construct(const FArguments& Args, EScriptSour
 	ParentArgs
 	.Style(FAppStyle::Get(), "ContentBrowser.FilterButton")
 	.IsChecked(Args._IsChecked)
-	.OnCheckStateChanged(FOnCheckStateChanged::CreateLambda([=](ECheckBoxState NewState)
+	.OnCheckStateChanged(FOnCheckStateChanged::CreateLambda([this](ECheckBoxState NewState)
     {
         OnSourceStateChanged.ExecuteIfBound(Source, NewState == ECheckBoxState::Checked ? true : false);
     }));
@@ -147,7 +147,7 @@ void SNiagaraSourceFilterBox::Construct(const FArguments& Args)
     
                 return bChecked ? ECheckBoxState::Checked : ECheckBoxState::Unchecked;
             })
-            .OnCheckStateChanged(FOnCheckStateChanged::CreateLambda([=](ECheckBoxState NewState)
+            .OnCheckStateChanged(FOnCheckStateChanged::CreateLambda([this, ScriptSourceEnum](ECheckBoxState NewState)
             {
             	bool bAnyChange = false;
             	// we always want to "Show all" so we always set the source filters to true
@@ -211,12 +211,12 @@ void SNiagaraSourceFilterBox::Construct(const FArguments& Args)
     			{
     				return SourceState[(EScriptSource)ScriptSourceEnum->GetValueByIndex(SourceIndex)] ? ECheckBoxState::Checked : ECheckBoxState::Unchecked;
     			})
-    			.OnSourceStateChanged_Lambda([=](EScriptSource Source, bool bState)
+    			.OnSourceStateChanged_Lambda([this](EScriptSource Source, bool bState)
     			{
     				SourceState.Add(Source, bState);
     				BroadcastFiltersChanged();
     			})
-    			.OnShiftClicked_Lambda([=](EScriptSource ChangedSource, bool bState)
+    			.OnShiftClicked_Lambda([this](EScriptSource ChangedSource, bool bState)
     			{
     				TArray<EScriptSource> Keys;
     				SourceState.GenerateKeyArray(Keys);
diff --git a/Engine/Plugins/FX/Niagara/Source/NiagaraEditor/Private/Widgets/SNiagaraHierarchy.cpp b/Engine/Plugins/FX/Niagara/Source/NiagaraEditor/Private/Widgets/SNiagaraHierarchy.cpp
--- a/Engine/Plugins/FX/Niagara/Source/NiagaraEditor/Private/Widgets/SNiagaraHierarchy.cpp
+++ b/Engine/Plugins/FX/Niagara/Source/NiagaraEditor/Private/Widgets/SNiagaraHierarchy.cpp
@@ -893,7 +893,7 @@ void SNiagaraHierarchy::RequestRefreshSourceViewNextFrame(bool bFullRefresh)
 {
 	if(!RefreshSourceViewNextFrameHandle.IsValid())
 	{
-		RefreshSourceViewNextFrameHandle = RegisterActiveTimer(0.f, FWidgetActiveTimerDelegate::CreateLambda([=](double CurrentTime, float DeltaTime)
+		RefreshSourceViewNextFrameHandle = RegisterActiveTimer(0.f, FWidgetActiveTimerDelegate::CreateLambda([this, bFullRefresh](double CurrentTime, float DeltaTime)
 		{
 			RefreshSourceView(bFullRefresh);
 			RefreshSourceViewNextFrameHandle.Reset();
@@ -920,7 +920,7 @@ void SNiagaraHierarchy::RequestRefreshHierarchyViewNextFrame(bool bFullRefresh)
 {
 	if(!RefreshHierarchyViewNextFrameHandle.IsValid())
 	{
-		RefreshHierarchyViewNextFrameHandle = RegisterActiveTimer(0.f, FWidgetActiveTimerDelegate::CreateLambda([=](double CurrentTime, float DeltaTime)
+		RefreshHierarchyViewNextFrameHandle = RegisterActiveTimer(0.f, FWidgetActiveTimerDelegate::CreateLambda([this, bFullRefresh](double CurrentTime, float DeltaTime)
 		{
 			RefreshHierarchyView(bFullRefresh);
 			RefreshHierarchyViewNextFrameHandle.Reset();
@@ -1015,7 +1015,7 @@ void SNiagaraHierarchy::RequestRefreshSectionsViewNextFrame()
 {
 	if(!RefreshSectionsViewNextFrameHandle.IsValid())
 	{
-		RefreshSectionsViewNextFrameHandle = RegisterActiveTimer(0.f, FWidgetActiveTimerDelegate::CreateLambda([=](double CurrentTime, float DeltaTime)
+		RefreshSectionsViewNextFrameHandle = RegisterActiveTimer(0.f, FWidgetActiveTimerDelegate::CreateLambda([this](double CurrentTime, float DeltaTime)
 		{
 			RefreshSectionsView();
 			RefreshSectionsViewNextFrameHandle.Reset();
diff --git a/Engine/Plugins/FX/Niagara/Source/NiagaraEditorWidgets/Private/SNiagaraOverviewGraph.cpp b/Engine/Plugins/FX/Niagara/Source/NiagaraEditorWidgets/Private/SNiagaraOverviewGraph.cpp
--- a/Engine/Plugins/FX/Niagara/Source/NiagaraEditorWidgets/Private/SNiagaraOverviewGraph.cpp
+++ b/Engine/Plugins/FX/Niagara/Source/NiagaraEditorWidgets/Private/SNiagaraOverviewGraph.cpp
@@ -106,13 +106,13 @@ void SNiagaraOverviewGraph::Construct(const FArguments& InArgs, TSharedRef<FNiag
 
 	Commands->MapAction(
 		FNiagaraEditorCommands::Get().HideDisabledModules,
-		FExecuteAction::CreateLambda([=]()
+		FExecuteAction::CreateLambda([this]()
 		{
 			UNiagaraStackEditorData& EditorData = ViewModel->GetSystemViewModel()->GetEditorData().GetStackEditorData();
 			EditorData.bHideDisabledModules = !EditorData.bHideDisabledModules;
 		}),
 		FCanExecuteAction(),
-		FIsActionChecked::CreateLambda([=]()
+		FIsActionChecked::CreateLambda([this]()
 		{
 			UNiagaraStackEditorData& EditorData = ViewModel->GetSystemViewModel()->GetEditorData().GetStackEditorData();
 			return EditorData.bHideDisabledModules;
diff --git a/Engine/Plugins/FX/Niagara/Source/NiagaraEditorWidgets/Private/SNiagaraOverviewStack.cpp b/Engine/Plugins/FX/Niagara/Source/NiagaraEditorWidgets/Private/SNiagaraOverviewStack.cpp
--- a/Engine/Plugins/FX/Niagara/Source/NiagaraEditorWidgets/Private/SNiagaraOverviewStack.cpp
+++ b/Engine/Plugins/FX/Niagara/Source/NiagaraEditorWidgets/Private/SNiagaraOverviewStack.cpp
@@ -131,7 +131,7 @@ class SNiagaraSystemOverviewEntryListRow : public STableRow<UNiagaraStackEntry*>
 						[					
 							SNew(SButton)
 							.ButtonStyle(FNiagaraEditorWidgetsStyle::Get(), "NiagaraEditor.Stack.SimpleButton")
-							.Visibility_Lambda([=]()
+							.Visibility_Lambda([this]()
 							{
 								if (StackEntry->GetCanExpandInOverview())
 								{	
@@ -144,7 +144,7 @@ class SNiagaraSystemOverviewEntryListRow : public STableRow<UNiagaraStackEntry*>
 									return EVisibility::Collapsed;
 								}
 							})
-							.OnClicked_Lambda([=]()
+							.OnClicked_Lambda([this]()
 							{
 								const bool bWillBeExpanded = !StackEntry->GetIsExpandedInOverview();
 								FText Message = FText::Format(LOCTEXT("ChangedCollapseState", "{0} {1}"), bWillBeExpanded ?
@@ -216,12 +216,12 @@ class SNiagaraSystemOverviewEntryListRow : public STableRow<UNiagaraStackEntry*>
 			ScalabilityOverlay->AddSlot()
 			[
 				SNew(SBorder)
-				.Visibility_Lambda([=]()
+				.Visibility_Lambda([this, RendererItem]()
 				{
 					return bScalabilityModeActive && RendererItem->IsExcludedFromScalability() && !RendererItem->IsOwningEmitterExcludedFromScalability() ? EVisibility::HitTestInvisible : EVisibility::Collapsed; 
 				})
 				.BorderImage(FNiagaraEditorStyle::Get().GetBrush("NiagaraEditor.SystemOverview.ExcludedFromScalability.RendererItem"))
-				.BorderBackgroundColor_Lambda([=]()
+				.BorderBackgroundColor_Lambda([this]()
 				{
 					return FLinearColor(1, 1, 1, FNiagaraEditorUtilities::GetScalabilityTintAlpha(EmitterHandleViewModel.Pin()->GetEmitterHandle()));
 				})
diff --git a/Engine/Plugins/FX/Niagara/Source/NiagaraEditorWidgets/Private/SNiagaraScalabilityPreviewSettings.cpp b/Engine/Plugins/FX/Niagara/Source/NiagaraEditorWidgets/Private/SNiagaraScalabilityPreviewSettings.cpp
--- a/Engine/Plugins/FX/Niagara/Source/NiagaraEditorWidgets/Private/SNiagaraScalabilityPreviewSettings.cpp
+++ b/Engine/Plugins/FX/Niagara/Source/NiagaraEditorWidgets/Private/SNiagaraScalabilityPreviewSettings.cpp
@@ -124,7 +124,7 @@ TSharedRef<SWidget> SNiagaraScalabilityPreviewSettings::CreatePreviewPlatformWid
 			.ButtonStyle(FAppStyle::Get(), "HoverHintOnly")
 			.OnClicked(this, &SNiagaraScalabilityPreviewSettings::OnResetPreviewPlatformClicked)
 			.ToolTipText(LOCTEXT("ResetPreviewPlatformToolTip", "Reset the platform to preview."))
-			.Visibility_Lambda([=]()
+			.Visibility_Lambda([this]()
 			{
 				return ScalabilityViewModel->GetPreviewDeviceProfile().IsSet() ? EVisibility::Visible : EVisibility::Collapsed;
 			})
@@ -150,7 +150,7 @@ TSharedRef<SWidget> SNiagaraScalabilityPreviewSettings::CreatePreviewPlatformWid
 					SNew(SBox)
 					.WidthOverride(16.f)
 					.HeightOverride(16.f)
-					.Visibility_Lambda([=]()
+					.Visibility_Lambda([this]()
 				    {
 						return GetActivePreviewPlatformImage() != nullptr
 							? EVisibility::Visible
diff --git a/Engine/Plugins/FX/Niagara/Source/NiagaraEditorWidgets/Private/SNiagaraStack.cpp b/Engine/Plugins/FX/Niagara/Source/NiagaraEditorWidgets/Private/SNiagaraStack.cpp
--- a/Engine/Plugins/FX/Niagara/Source/NiagaraEditorWidgets/Private/SNiagaraStack.cpp
+++ b/Engine/Plugins/FX/Niagara/Source/NiagaraEditorWidgets/Private/SNiagaraStack.cpp
@@ -549,9 +549,9 @@ TSharedRef<SWidget> SNiagaraStack::GetViewOptionsMenu() const
 		LOCTEXT("ShowAllAdvancedToolTip", "Forces all advanced items to be showing in the stack."),
 		FSlateIcon(),
 		FUIAction(
-			FExecuteAction::CreateLambda([=]() { StackViewModel->SetShowAllAdvanced(!StackViewModel->GetShowAllAdvanced()); }),
+			FExecuteAction::CreateLambda([this]() { StackViewModel->SetShowAllAdvanced(!StackViewModel->GetShowAllAdvanced()); }),
 			FCanExecuteAction(),
-			FGetActionCheckState::CreateLambda([=]() { return StackViewModel->GetShowAllAdvanced() ? ECheckBoxState::Checked : ECheckBoxState::Unchecked; })),
+			FGetActionCheckState::CreateLambda([this]() { return StackViewModel->GetShowAllAdvanced() ? ECheckBoxState::Checked : ECheckBoxState::Unchecked; })),
 		NAME_None, EUserInterfaceActionType::Check);
 
 	MenuBuilder.AddMenuEntry(
@@ -559,9 +559,9 @@ TSharedRef<SWidget> SNiagaraStack::GetViewOptionsMenu() const
 		LOCTEXT("ShowParameterReadsToolTip", "Whether or not to show the parameters that a module reads from."),
 		FSlateIcon(),
 		FUIAction(
-			FExecuteAction::CreateLambda([=]() { StackViewModel->SetShowLinkedInputs(!StackViewModel->GetShowLinkedInputs()); }),
+			FExecuteAction::CreateLambda([this]() { StackViewModel->SetShowLinkedInputs(!StackViewModel->GetShowLinkedInputs()); }),
 			FCanExecuteAction(),
-			FGetActionCheckState::CreateLambda([=]() { return StackViewModel->GetShowLinkedInputs() ? ECheckBoxState::Checked : ECheckBoxState::Unchecked; })),
+			FGetActionCheckState::CreateLambda([this]() { return StackViewModel->GetShowLinkedInputs() ? ECheckBoxState::Checked : ECheckBoxState::Unchecked; })),
 		NAME_None, EUserInterfaceActionType::Check);
 
 	MenuBuilder.AddMenuEntry(
@@ -569,9 +569,9 @@ TSharedRef<SWidget> SNiagaraStack::GetViewOptionsMenu() const
 		LOCTEXT("ShowParameterWritesToolTip", "Whether or not to show parameters that a module writes to."),
 		FSlateIcon(),
 		FUIAction(
-			FExecuteAction::CreateLambda([=]() { StackViewModel->SetShowOutputs(!StackViewModel->GetShowOutputs()); }),
+			FExecuteAction::CreateLambda([this]() { StackViewModel->SetShowOutputs(!StackViewModel->GetShowOutputs()); }),
 			FCanExecuteAction(),
-			FGetActionCheckState::CreateLambda([=]() { return StackViewModel->GetShowOutputs() ? ECheckBoxState::Checked : ECheckBoxState::Unchecked; })),
+			FGetActionCheckState::CreateLambda([this]() { return StackViewModel->GetShowOutputs() ? ECheckBoxState::Checked : ECheckBoxState::Unchecked; })),
 		NAME_None, EUserInterfaceActionType::Check);
 
 	MenuBuilder.AddMenuEntry(
@@ -579,9 +579,9 @@ TSharedRef<SWidget> SNiagaraStack::GetViewOptionsMenu() const
 		LOCTEXT("ShowIssuesToolTip", "Hides all modules except those that have unresolved issues."),
 		FSlateIcon(),
 		FUIAction(
-			FExecuteAction::CreateLambda([=]() { StackViewModel->SetShowOnlyIssues(!StackViewModel->GetShowOnlyIssues()); }),
+			FExecuteAction::CreateLambda([this]() { StackViewModel->SetShowOnlyIssues(!StackViewModel->GetShowOnlyIssues()); }),
 			FCanExecuteAction(),
-			FGetActionCheckState::CreateLambda([=]() { return StackViewModel->GetShowOnlyIssues() ? ECheckBoxState::Checked : ECheckBoxState::Unchecked; })),
+			FGetActionCheckState::CreateLambda([this]() { return StackViewModel->GetShowOnlyIssues() ? ECheckBoxState::Checked : ECheckBoxState::Unchecked; })),
 		NAME_None, EUserInterfaceActionType::Check);
 
 	return MenuBuilder.MakeWidget();
diff --git a/Engine/Plugins/FX/Niagara/Source/NiagaraEditorWidgets/Private/Stack/SNiagaraStackModuleItem.cpp b/Engine/Plugins/FX/Niagara/Source/NiagaraEditorWidgets/Private/Stack/SNiagaraStackModuleItem.cpp
--- a/Engine/Plugins/FX/Niagara/Source/NiagaraEditorWidgets/Private/Stack/SNiagaraStackModuleItem.cpp
+++ b/Engine/Plugins/FX/Niagara/Source/NiagaraEditorWidgets/Private/Stack/SNiagaraStackModuleItem.cpp
@@ -41,7 +41,7 @@ void SNiagaraStackModuleItem::Construct(const FArguments& InArgs, UNiagaraStackM
 {
 	ModuleItem = &InModuleItem;
 
-	ModuleItem->OnNoteModeSet().BindLambda([=](bool bEnabled)
+	ModuleItem->OnNoteModeSet().BindLambda([this](bool bEnabled)
 	{
 		if(bEnabled)
 		{
diff --git a/Engine/Plugins/Media/MediaPlayerEditor/Source/MediaPlayerEditor/Private/Widgets/SMediaPlayerEditorViewer.cpp b/Engine/Plugins/Media/MediaPlayerEditor/Source/MediaPlayerEditor/Private/Widgets/SMediaPlayerEditorViewer.cpp
--- a/Engine/Plugins/Media/MediaPlayerEditor/Source/MediaPlayerEditor/Private/Widgets/SMediaPlayerEditorViewer.cpp
+++ b/Engine/Plugins/Media/MediaPlayerEditor/Source/MediaPlayerEditor/Private/Widgets/SMediaPlayerEditorViewer.cpp
@@ -740,9 +740,9 @@ void SMediaPlayerEditorViewer::HandleDecoderMenuNewMenu(FMenuBuilder& MenuBuilde
 			FText::FromString(FString::Join(Factory->GetSupportedPlatforms(), TEXT(", "))),
 			FSlateIcon(),
 			FUIAction(
-				FExecuteAction::CreateLambda([=] { SetDesiredPlayerName(PlayerName); }),
-				FCanExecuteAction::CreateLambda([=] { return SupportsRunningPlatform; }),
-				FIsActionChecked::CreateLambda([=] { return MediaPlayer->GetDesiredPlayerName() == PlayerName; })
+				FExecuteAction::CreateLambda([this, PlayerName] { SetDesiredPlayerName(PlayerName); }),
+				FCanExecuteAction::CreateLambda([SupportsRunningPlatform] { return SupportsRunningPlatform; }),
+				FIsActionChecked::CreateLambda([this, PlayerName] { return MediaPlayer->GetDesiredPlayerName() == PlayerName; })
 			),
 			NAME_None,
 			EUserInterfaceActionType::RadioButton
@@ -811,9 +811,9 @@ void SMediaPlayerEditorViewer::HandleFormatMenuNewMenu(FMenuBuilder& MenuBuilder
 				FText::GetEmpty(),
 				FSlateIcon(),
 				FUIAction(
-					FExecuteAction::CreateLambda([=] { MediaPlayer->SetTrackFormat(TrackType, SelectedTrack, FormatIndex); }),
+					FExecuteAction::CreateLambda([this, TrackType, SelectedTrack, FormatIndex] { MediaPlayer->SetTrackFormat(TrackType, SelectedTrack, FormatIndex); }),
 					FCanExecuteAction(),
-					FIsActionChecked::CreateLambda([=] { return (MediaPlayer->GetTrackFormat(TrackType, SelectedTrack) == FormatIndex); })
+					FIsActionChecked::CreateLambda([this, TrackType, SelectedTrack, FormatIndex] { return (MediaPlayer->GetTrackFormat(TrackType, SelectedTrack) == FormatIndex); })
 				),
 				NAME_None,
 				EUserInterfaceActionType::RadioButton
@@ -1004,9 +1004,9 @@ void SMediaPlayerEditorViewer::HandleTrackMenuNewMenu(FMenuBuilder& MenuBuilder,
 			FText::GetEmpty(),
 			FSlateIcon(),
 			FUIAction(
-				FExecuteAction::CreateLambda([=] { MediaPlayer->SelectTrack(TrackType, INDEX_NONE); }),
+				FExecuteAction::CreateLambda([this, TrackType] { MediaPlayer->SelectTrack(TrackType, INDEX_NONE); }),
 				FCanExecuteAction(),
-				FIsActionChecked::CreateLambda([=] { return (MediaPlayer->GetSelectedTrack(TrackType) == INDEX_NONE); })
+				FIsActionChecked::CreateLambda([this, TrackType] { return (MediaPlayer->GetSelectedTrack(TrackType) == INDEX_NONE); })
 			),
 			NAME_None,
 			EUserInterfaceActionType::RadioButton
@@ -1032,9 +1032,9 @@ void SMediaPlayerEditorViewer::HandleTrackMenuNewMenu(FMenuBuilder& MenuBuilder,
 				TooltipText,
 				FSlateIcon(),
 				FUIAction(
-					FExecuteAction::CreateLambda([=] { MediaPlayer->SelectTrack(TrackType, TrackIndex); }),
+					FExecuteAction::CreateLambda([this, TrackType, TrackIndex] { MediaPlayer->SelectTrack(TrackType, TrackIndex); }),
 					FCanExecuteAction(),
-					FIsActionChecked::CreateLambda([=] { return (MediaPlayer->GetSelectedTrack(TrackType) == TrackIndex); })
+					FIsActionChecked::CreateLambda([this, TrackType, TrackIndex] { return (MediaPlayer->GetSelectedTrack(TrackType) == TrackIndex); })
 				),
 				NAME_None,
 				EUserInterfaceActionType::RadioButton
diff --git a/Engine/Plugins/MovieScene/ActorSequence/Source/ActorSequenceEditor/Private/ActorSequenceEditorTabSummoner.cpp b/Engine/Plugins/MovieScene/ActorSequence/Source/ActorSequenceEditor/Private/ActorSequenceEditorTabSummoner.cpp
--- a/Engine/Plugins/MovieScene/ActorSequence/Source/ActorSequenceEditor/Private/ActorSequenceEditorTabSummoner.cpp
+++ b/Engine/Plugins/MovieScene/ActorSequence/Source/ActorSequenceEditor/Private/ActorSequenceEditorTabSummoner.cpp
@@ -387,7 +387,7 @@ public:
 			TSharedRef<FExtender> AddMenuExtender = MakeShareable(new FExtender);
 
 			AddMenuExtender->AddMenuExtension("AddTracks", EExtensionHook::Before, nullptr,
-				FMenuExtensionDelegate::CreateLambda([=](FMenuBuilder& MenuBuilder){
+				FMenuExtensionDelegate::CreateLambda([this](FMenuBuilder& MenuBuilder){
 
 					MenuBuilder.AddSubMenu(
 						LOCTEXT("AddComponent_Label", "Component"),
diff --git a/Engine/Plugins/MovieScene/LevelSequenceEditor/Source/LevelSequenceEditor/Private/CinematicViewport/SCinematicLevelViewport.cpp b/Engine/Plugins/MovieScene/LevelSequenceEditor/Source/LevelSequenceEditor/Private/CinematicViewport/SCinematicLevelViewport.cpp
--- a/Engine/Plugins/MovieScene/LevelSequenceEditor/Source/LevelSequenceEditor/Private/CinematicViewport/SCinematicLevelViewport.cpp
+++ b/Engine/Plugins/MovieScene/LevelSequenceEditor/Source/LevelSequenceEditor/Private/CinematicViewport/SCinematicLevelViewport.cpp
@@ -225,7 +225,7 @@ void SCinematicLevelViewport::Construct(const FArguments& InArgs)
 				.OnEndSliderMovement(this, &SCinematicLevelViewport::SetTime)
 				.Value(this, &SCinematicLevelViewport::GetTime)
 				.ToolTipText(LOCTEXT("TimeLocalToCurrentSequence", "The current time of the sequence relative to the focused sequence."))
-				.Delta_Lambda([=]()
+				.Delta_Lambda([this]()
 				{
 					return UIData.OuterResolution.AsDecimal() * UIData.OuterPlayRate.AsInterval(); 
 				})
@@ -337,7 +337,7 @@ void SCinematicLevelViewport::Construct(const FArguments& InArgs)
 								[
 									SNew(STextBlock)
 									.ColorAndOpacity(Gray)
-									.Text_Lambda([=]{ return UIData.ShotName; })
+									.Text_Lambda([this]{ return UIData.ShotName; })
 									.ToolTipText(LOCTEXT("CurrentSequence", "The name of the currently evaluated sequence."))
 								]
 
@@ -348,7 +348,7 @@ void SCinematicLevelViewport::Construct(const FArguments& InArgs)
 								[
 									SNew(STextBlock)
 									.ColorAndOpacity(Gray)
-									.Text_Lambda([=] { return UIData.CameraName; })
+									.Text_Lambda([this] { return UIData.CameraName; })
 									.ToolTipText(LOCTEXT("CurrentCamera", "The name of the current camera."))
 								]
 							]
@@ -359,7 +359,7 @@ void SCinematicLevelViewport::Construct(const FArguments& InArgs)
 							[
 								SNew(STextBlock)
 								.ColorAndOpacity(Gray)
-								.Text_Lambda([=] { return UIData.Filmback; })
+								.Text_Lambda([this] { return UIData.Filmback; })
 								.ToolTipText(LOCTEXT("CurrentFilmback", "The name of the current shot's filmback (the imaging area of the frame/sensor)."))
 							]
 
@@ -369,7 +369,7 @@ void SCinematicLevelViewport::Construct(const FArguments& InArgs)
 								SNew(STextBlock)
 								.Font(FAppStyle::GetFontStyle("Sequencer.FixedFont"))
 								.ColorAndOpacity(Gray)
-								.Text_Lambda([=] { return UIData.LocalPlaybackTime; })
+								.Text_Lambda([this] { return UIData.LocalPlaybackTime; })
 								.ToolTipText(LOCTEXT("LocalPlaybackTime", "The current playback time relative to the currently evaluated sequence."))
 							]
 						]
diff --git a/Engine/Plugins/MovieScene/LevelSequenceEditor/Source/LevelSequenceEditor/Private/Misc/LevelSequenceEditorActorBinding.cpp b/Engine/Plugins/MovieScene/LevelSequenceEditor/Source/LevelSequenceEditor/Private/Misc/LevelSequenceEditorActorBinding.cpp
--- a/Engine/Plugins/MovieScene/LevelSequenceEditor/Source/LevelSequenceEditor/Private/Misc/LevelSequenceEditorActorBinding.cpp
+++ b/Engine/Plugins/MovieScene/LevelSequenceEditor/Source/LevelSequenceEditor/Private/Misc/LevelSequenceEditorActorBinding.cpp
@@ -92,7 +92,7 @@ void FLevelSequenceEditorActorBinding::AddPossessActorMenuExtensions(FMenuBuilde
 	if (!SelectedLabel.IsEmpty())
 	{
 		// Copy the array into the lambda - probably not that big a deal
-		MenuBuilder.AddMenuEntry(SelectedLabel, FText(), ActorIcon, FExecuteAction::CreateLambda([=]{
+		MenuBuilder.AddMenuEntry(SelectedLabel, FText(), ActorIcon, FExecuteAction::CreateLambda([this, ActorsValidForPossession]{
 			FSlateApplication::Get().DismissAllMenus();
 			AddActorsToSequencer(ActorsValidForPossession.GetData(), ActorsValidForPossession.Num());
 		}));
@@ -128,7 +128,7 @@ void FLevelSequenceEditorActorBinding::AddPossessActorMenuExtensions(FMenuBuilde
 		[
 			SceneOutlinerModule.CreateActorPicker(
 				InitOptions,
-				FOnActorPicked::CreateLambda([=](AActor* Actor){
+				FOnActorPicked::CreateLambda([this](AActor* Actor){
 					// Create a new binding for this actor
 					FSlateApplication::Get().DismissAllMenus();
 					AddActorsToSequencer(&Actor, 1);
diff --git a/Engine/Plugins/Runtime/HairStrands/Source/HairStrandsEditor/Private/GroomCacheTrackEditor.cpp b/Engine/Plugins/Runtime/HairStrands/Source/HairStrandsEditor/Private/GroomCacheTrackEditor.cpp
--- a/Engine/Plugins/Runtime/HairStrands/Source/HairStrandsEditor/Private/GroomCacheTrackEditor.cpp
+++ b/Engine/Plugins/Runtime/HairStrands/Source/HairStrandsEditor/Private/GroomCacheTrackEditor.cpp
@@ -379,7 +379,7 @@ TSharedPtr<SWidget> FGroomCacheTrackEditor::BuildOutlinerEditWidget(const FGuid&
 	{
 		TWeakPtr<ISequencer> WeakSequencer = GetSequencer();
 
-		auto SubMenuCallback = [=]() -> TSharedRef<SWidget>
+		auto SubMenuCallback = [this, ObjectBinding, Track]() -> TSharedRef<SWidget>
 		{
 			FMenuBuilder MenuBuilder(true, nullptr);
 
diff --git a/Engine/Plugins/Runtime/Metasound/Source/MetasoundEditor/Private/MetasoundEditor.cpp b/Engine/Plugins/Runtime/Metasound/Source/MetasoundEditor/Private/MetasoundEditor.cpp
--- a/Engine/Plugins/Runtime/Metasound/Source/MetasoundEditor/Private/MetasoundEditor.cpp
+++ b/Engine/Plugins/Runtime/Metasound/Source/MetasoundEditor/Private/MetasoundEditor.cpp
@@ -149,7 +149,7 @@ namespace Metasound
 			virtual int32 GetReorderIndexInContainer() const override
 			{
 				TArray<Frontend::FConstNodeHandle> InputHandles = GetGraphHandle()->GetConstInputNodes();
-				return InputHandles.IndexOfByPredicate([=](const Frontend::FConstNodeHandle& NodeHandle)
+				return InputHandles.IndexOfByPredicate([this](const Frontend::FConstNodeHandle& NodeHandle)
 				{
 					return NodeHandle->GetID() == MemberID;
 				});
diff --git a/Engine/Plugins/Runtime/nDisplay/Source/DisplayClusterConfigurator/Private/ClusterConfiguration/SDisplayClusterConfiguratorNewClusterItemDialog.cpp b/Engine/Plugins/Runtime/nDisplay/Source/DisplayClusterConfigurator/Private/ClusterConfiguration/SDisplayClusterConfiguratorNewClusterItemDialog.cpp
--- a/Engine/Plugins/Runtime/nDisplay/Source/DisplayClusterConfigurator/Private/ClusterConfiguration/SDisplayClusterConfiguratorNewClusterItemDialog.cpp
+++ b/Engine/Plugins/Runtime/nDisplay/Source/DisplayClusterConfigurator/Private/ClusterConfiguration/SDisplayClusterConfiguratorNewClusterItemDialog.cpp
@@ -206,7 +206,7 @@ void SDisplayClusterConfiguratorNewClusterItemDialog::Construct(const FArguments
 								.OptionsSource(&PresetItems)
 								.InitiallySelectedItem(InitiallySelectedPresetItem)
 								.OnSelectionChanged(this, &SDisplayClusterConfiguratorNewClusterItemDialog::OnSelectedPresetChanged)
-								.OnGenerateWidget_Lambda([=](TSharedPtr<FDisplayClusterConfiguratorPresetSize> Item)
+								.OnGenerateWidget_Lambda([this](TSharedPtr<FDisplayClusterConfiguratorPresetSize> Item)
 								{
 									return SNew(STextBlock).Text(GetPresetDisplayText(Item));
 								})
diff --git a/Engine/Plugins/Runtime/nDisplay/Source/DisplayClusterConfigurator/Private/DisplayClusterConfiguratorBlueprintEditor.cpp b/Engine/Plugins/Runtime/nDisplay/Source/DisplayClusterConfigurator/Private/DisplayClusterConfiguratorBlueprintEditor.cpp
--- a/Engine/Plugins/Runtime/nDisplay/Source/DisplayClusterConfigurator/Private/DisplayClusterConfiguratorBlueprintEditor.cpp
+++ b/Engine/Plugins/Runtime/nDisplay/Source/DisplayClusterConfigurator/Private/DisplayClusterConfiguratorBlueprintEditor.cpp
@@ -1318,7 +1318,7 @@ void FDisplayClusterConfiguratorBlueprintEditor::ShutdownDCSCSEditors()
 
 	// Remove this blueprint editor and any registered blueprint editors that are no longer valid (as is the case when in a destructor)
 	// from the list of registered blueprint editors
-	EditorsRegisteredForPanelExtension.RemoveAll([=](const TWeakPtr<FDisplayClusterConfiguratorBlueprintEditor>& BlueprintEditor)
+	EditorsRegisteredForPanelExtension.RemoveAll([this](const TWeakPtr<FDisplayClusterConfiguratorBlueprintEditor>& BlueprintEditor)
 	{
 		return !BlueprintEditor.IsValid() || BlueprintEditor.Pin().Get() == this;
 	});
@@ -1326,7 +1326,7 @@ void FDisplayClusterConfiguratorBlueprintEditor::ShutdownDCSCSEditors()
 
 TSharedRef<SWidget> FDisplayClusterConfiguratorBlueprintEditor::CreateAddComponentButton()
 {
-	auto PerformComboAddClass = [=](TSubclassOf<UActorComponent> ComponentClass, EComponentCreateAction::Type ComponentCreateAction,
+	auto PerformComboAddClass = [this](TSubclassOf<UActorComponent> ComponentClass, EComponentCreateAction::Type ComponentCreateAction,
 		UObject* AssetOverride) -> FSubobjectDataHandle
 	{
 		FSubobjectDataHandle NewComponentHandle = FSubobjectDataHandle::InvalidHandle;
diff --git a/Engine/Plugins/Runtime/nDisplay/Source/DisplayClusterConfigurator/Private/Views/Details/Cluster/DisplayClusterConfiguratorExternalImageTypeCustomization.cpp b/Engine/Plugins/Runtime/nDisplay/Source/DisplayClusterConfigurator/Private/Views/Details/Cluster/DisplayClusterConfiguratorExternalImageTypeCustomization.cpp
--- a/Engine/Plugins/Runtime/nDisplay/Source/DisplayClusterConfigurator/Private/Views/Details/Cluster/DisplayClusterConfiguratorExternalImageTypeCustomization.cpp
+++ b/Engine/Plugins/Runtime/nDisplay/Source/DisplayClusterConfigurator/Private/Views/Details/Cluster/DisplayClusterConfiguratorExternalImageTypeCustomization.cpp
@@ -36,6 +36,6 @@ void FDisplayClusterConfiguratorExternalImageTypeCustomization::CustomizeHeader(
 		SNew(SDisplayClusterConfiguratorExternalImagePicker)
 		.ImagePath(ImagePath)
 		.Extensions(ImageExtensions)
-		.OnImagePathPicked_Lambda([=](const FString& NewImagePath) { ImagePathHandle->SetValue(NewImagePath); })
+		.OnImagePathPicked_Lambda([this](const FString& NewImagePath) { ImagePathHandle->SetValue(NewImagePath); })
 	];
 }
\ No newline at end of file
diff --git a/Engine/Plugins/Runtime/nDisplay/Source/DisplayClusterConfigurator/Private/Views/Details/Components/DisplayClusterConfiguratorScreenComponentDetailsCustomization.cpp b/Engine/Plugins/Runtime/nDisplay/Source/DisplayClusterConfigurator/Private/Views/Details/Components/DisplayClusterConfiguratorScreenComponentDetailsCustomization.cpp
--- a/Engine/Plugins/Runtime/nDisplay/Source/DisplayClusterConfigurator/Private/Views/Details/Components/DisplayClusterConfiguratorScreenComponentDetailsCustomization.cpp
+++ b/Engine/Plugins/Runtime/nDisplay/Source/DisplayClusterConfigurator/Private/Views/Details/Components/DisplayClusterConfiguratorScreenComponentDetailsCustomization.cpp
@@ -79,7 +79,7 @@ void FDisplayClusterConfiguratorScreenDetailsCustomization::CustomizeDetails(IDe
 				.OptionsSource(&PresetItems)
 				.InitiallySelectedItem(InitiallySelectedPresetItem)
 				.OnSelectionChanged(this, &FDisplayClusterConfiguratorScreenDetailsCustomization::OnSelectedPresetChanged)
-				.OnGenerateWidget_Lambda([=](TSharedPtr<FDisplayClusterConfiguratorAspectRatioPresetSize> Item)
+				.OnGenerateWidget_Lambda([this](TSharedPtr<FDisplayClusterConfiguratorAspectRatioPresetSize> Item)
 				{
 					return SNew(STextBlock)
 						.Font(IDetailLayoutBuilder::GetDetailFont())
diff --git a/Engine/Plugins/Runtime/nDisplay/Source/DisplayClusterLightCardEditor/Private/DisplayClusterLightCardEditor.cpp b/Engine/Plugins/Runtime/nDisplay/Source/DisplayClusterLightCardEditor/Private/DisplayClusterLightCardEditor.cpp
--- a/Engine/Plugins/Runtime/nDisplay/Source/DisplayClusterLightCardEditor/Private/DisplayClusterLightCardEditor.cpp
+++ b/Engine/Plugins/Runtime/nDisplay/Source/DisplayClusterLightCardEditor/Private/DisplayClusterLightCardEditor.cpp
@@ -1914,7 +1914,7 @@ void FDisplayClusterLightCardEditor::OnLevelActorDeleted(AActor* Actor)
 		
 		if (Actor->GetWorld())
 		{
-			Actor->GetWorld()->GetTimerManager().SetTimerForNextTick([=]()
+			Actor->GetWorld()->GetTimerManager().SetTimerForNextTick([this]()
 			{
 				// Schedule for next tick so available selections are properly updated once the
 				// actor is fully deleted.
diff --git a/Engine/Plugins/Runtime/nDisplay/Source/DisplayClusterLightCardEditor/Private/Viewport/DisplayClusterLightCardEditorViewportClient.cpp b/Engine/Plugins/Runtime/nDisplay/Source/DisplayClusterLightCardEditor/Private/Viewport/DisplayClusterLightCardEditorViewportClient.cpp
--- a/Engine/Plugins/Runtime/nDisplay/Source/DisplayClusterLightCardEditor/Private/Viewport/DisplayClusterLightCardEditorViewportClient.cpp
+++ b/Engine/Plugins/Runtime/nDisplay/Source/DisplayClusterLightCardEditor/Private/Viewport/DisplayClusterLightCardEditorViewportClient.cpp
@@ -1637,7 +1637,7 @@ void FDisplayClusterLightCardEditorViewportClient::UpdatePreviewActor(ADisplayCl
 		TWeakObjectPtr<AActor> StageActorWeakPtr(StageActor);
 		// Schedule for the next tick so CDO changes get propagated first in the event of config editor skeleton
 		// regeneration & compiles. nDisplay's custom propagation may have issues if the archetype isn't correct.
-		PreviewWorld->GetTimerManager().SetTimerForNextTick([=]()
+		PreviewWorld->GetTimerManager().SetTimerForNextTick([=, this]()
 		{
 			DECLARE_SCOPE_CYCLE_COUNTER(TEXT("FDisplayClusterLightCardEditorViewportClient::UpdatePreviewActorImpl"), STAT_UpdatePreviewActorImpl, STATGROUP_NDisplayLightCardEditor);
 			
diff --git a/Engine/Plugins/Tests/EditorTests/Source/EditorTests/Private/UnrealEd/AutoReimportTests.cpp b/Engine/Plugins/Tests/EditorTests/Source/EditorTests/Private/UnrealEd/AutoReimportTests.cpp
--- a/Engine/Plugins/Tests/EditorTests/Source/EditorTests/Private/UnrealEd/AutoReimportTests.cpp
+++ b/Engine/Plugins/Tests/EditorTests/Source/EditorTests/Private/UnrealEd/AutoReimportTests.cpp
@@ -97,7 +97,7 @@ struct FAutoReimportTestPayload
 	void WaitForStartup(TFunction<void()>&& Finished)
 	{
 		// Ideally, this would capture-by-move, but we don't have full compiler support for that yet.
-		ADD_LATENT_AUTOMATION_COMMAND(FDelayedCallbackLatentCommand([=]() mutable {
+		ADD_LATENT_AUTOMATION_COMMAND(FDelayedCallbackLatentCommand([this, Finished]() mutable {
 			FileCache->Tick();
 			if (FileCache->MoveDetectionInitialized())
 			{
diff --git a/Engine/Plugins/VirtualProduction/CameraCalibration/Source/CameraCalibrationEditor/Private/AssetEditor/SCameraCalibrationSteps.cpp b/Engine/Plugins/VirtualProduction/CameraCalibration/Source/CameraCalibrationEditor/Private/AssetEditor/SCameraCalibrationSteps.cpp
--- a/Engine/Plugins/VirtualProduction/CameraCalibration/Source/CameraCalibrationEditor/Private/AssetEditor/SCameraCalibrationSteps.cpp
+++ b/Engine/Plugins/VirtualProduction/CameraCalibration/Source/CameraCalibrationEditor/Private/AssetEditor/SCameraCalibrationSteps.cpp
@@ -606,7 +606,7 @@ void SCameraCalibrationSteps::UpdateOverlayMaterialParameterWidget()
 						.MinValue(TOptional<float>())
 						.MaxValue(TOptional<float>())
 						.Delta(0.0f)
-						.Value_Lambda([=]() 
+						.Value_Lambda([this, ParameterInfo]() 
 						{ 
 							float ScalarValue = 0.0f;
 							
@@ -617,7 +617,7 @@ void SCameraCalibrationSteps::UpdateOverlayMaterialParameterWidget()
 
 							return ScalarValue;
 						})
-						.MinSliderValue_Lambda([=]() 
+						.MinSliderValue_Lambda([this, ParameterInfo]()
 						{ 
 							float MinValue = 0.0f;
 							float MaxValue = 0.0f;
@@ -629,7 +629,7 @@ void SCameraCalibrationSteps::UpdateOverlayMaterialParameterWidget()
 
 							return MinValue;
 						})
-						.MaxSliderValue_Lambda([=]() 
+						.MaxSliderValue_Lambda([this, ParameterInfo]()
 						{ 
 							float MinValue = 0.0f;
 							float MaxValue = 0.0f;
@@ -641,7 +641,7 @@ void SCameraCalibrationSteps::UpdateOverlayMaterialParameterWidget()
 
 							return MaxValue;
 						})
-						.OnValueChanged_Lambda([=](float NewValue) 
+						.OnValueChanged_Lambda([this, ParameterInfo](float NewValue)
 						{ 
 							if (UMaterialInstanceDynamic* Overlay = CurrentOverlayMID.Get())
 							{
@@ -653,7 +653,7 @@ void SCameraCalibrationSteps::UpdateOverlayMaterialParameterWidget()
 								}
 							}
 						})
-						.OnValueCommitted_Lambda([=](float NewValue, ETextCommit::Type)	
+						.OnValueCommitted_Lambda([this, ParameterInfo](float NewValue, ETextCommit::Type)
 						{
 							if (UMaterialInstanceDynamic* Overlay = CurrentOverlayMID.Get())
 							{
@@ -694,11 +694,11 @@ void SCameraCalibrationSteps::UpdateOverlayMaterialParameterWidget()
 						.HasDownArrow(false)
 						.CollapseMenuOnParentFocus(true)
 						.ButtonStyle(FAppStyle::Get(), "Sequencer.AnimationOutliner.ColorStrip") // Style matches the button used in cinematic film overlays
-						.OnGetMenuContent_Lambda([=]() -> TSharedRef<SWidget>
+						.OnGetMenuContent_Lambda([this, ParameterInfo]() -> TSharedRef<SWidget>
 						{
 							return SNew(SColorPicker)
 								.UseAlpha(true)
-								.TargetColorAttribute_Lambda([=]() 
+								.TargetColorAttribute_Lambda([this, ParameterInfo]()
 								{
 									FLinearColor ColorValue = FLinearColor(1.0f, 1.0f, 1.0f, 1.0f);
 										
@@ -709,7 +709,7 @@ void SCameraCalibrationSteps::UpdateOverlayMaterialParameterWidget()
 
 									return ColorValue;
 								})
-								.OnColorCommitted_Lambda([=](FLinearColor NewColor) 
+								.OnColorCommitted_Lambda([this, ParameterInfo](FLinearColor NewColor)
 								{
 									if (UMaterialInstanceDynamic* Overlay = CurrentOverlayMID.Get())
 									{
@@ -727,7 +727,7 @@ void SCameraCalibrationSteps::UpdateOverlayMaterialParameterWidget()
 							SNew(SColorBlock)
 							.ShowBackgroundForAlpha(true)
 							.Size(FVector2D(10.0f, 20.0f))
-							.Color_Lambda([=]()
+							.Color_Lambda([this, ParameterInfo]()
 							{
 								FLinearColor ColorValue = FLinearColor(1.0f, 1.0f, 1.0f, 1.0f);
 
diff --git a/Engine/Plugins/VirtualProduction/CameraCalibration/Source/CameraCalibrationEditor/Private/AssetEditor/SLensDataAddPointDialog.cpp b/Engine/Plugins/VirtualProduction/CameraCalibration/Source/CameraCalibrationEditor/Private/AssetEditor/SLensDataAddPointDialog.cpp
--- a/Engine/Plugins/VirtualProduction/CameraCalibration/Source/CameraCalibrationEditor/Private/AssetEditor/SLensDataAddPointDialog.cpp
+++ b/Engine/Plugins/VirtualProduction/CameraCalibration/Source/CameraCalibrationEditor/Private/AssetEditor/SLensDataAddPointDialog.cpp
@@ -158,7 +158,7 @@ void SLensDataAddPointDialog::Construct(const FArguments& InArgs, ULensFile* InL
 					.ButtonContent()
 					[
 						SNew(STextBlock)
-						.Text(MakeAttributeLambda([=]
+						.Text(MakeAttributeLambda([this]
 						{
 							switch (SelectedCategory)
 							{
diff --git a/Engine/Plugins/VirtualProduction/CameraCalibration/Source/CameraCalibrationEditor/Private/Calibrators/CameraNodalOffsetAlgoPoints.cpp b/Engine/Plugins/VirtualProduction/CameraCalibration/Source/CameraCalibrationEditor/Private/Calibrators/CameraNodalOffsetAlgoPoints.cpp
--- a/Engine/Plugins/VirtualProduction/CameraCalibration/Source/CameraCalibrationEditor/Private/Calibrators/CameraNodalOffsetAlgoPoints.cpp
+++ b/Engine/Plugins/VirtualProduction/CameraCalibration/Source/CameraCalibrationEditor/Private/Calibrators/CameraNodalOffsetAlgoPoints.cpp
@@ -1460,7 +1460,7 @@ TSharedRef<SWidget> UCameraNodalOffsetAlgoPoints::BuildCalibrationComponentPicke
 		.Padding(FMargin(0.0f, 0.0f, 0.0f, 0.0f))
 		[
 			SNew(SComboButton)
-			.OnGetMenuContent_Lambda([=]() { return BuildCalibrationComponentMenu(); })
+			.OnGetMenuContent_Lambda([this]() { return BuildCalibrationComponentMenu(); })
 			.ContentPadding(FMargin(4.0, 2.0))
 			.ButtonContent()
 			[
diff --git a/Engine/Plugins/VirtualProduction/DMX/DMXEngine/Source/DMXEditor/Private/Customizations/DMXAttributeNameCustomization.cpp b/Engine/Plugins/VirtualProduction/DMX/DMXEngine/Source/DMXEditor/Private/Customizations/DMXAttributeNameCustomization.cpp
--- a/Engine/Plugins/VirtualProduction/DMX/DMXEngine/Source/DMXEditor/Private/Customizations/DMXAttributeNameCustomization.cpp
+++ b/Engine/Plugins/VirtualProduction/DMX/DMXEngine/Source/DMXEditor/Private/Customizations/DMXAttributeNameCustomization.cpp
@@ -40,7 +40,7 @@ void FDMXAttributeNameCustomization::CustomizeHeader(TSharedRef<IPropertyHandle>
 	}
 
 	InHeaderRow
-		.IsEnabled(MakeAttributeLambda([=] 
+		.IsEnabled(MakeAttributeLambda([this, InPropertyHandle] 
 			{ 
 				return !InPropertyHandle->IsEditConst() && PropertyUtilities->IsPropertyEditingEnabled();
 			}))
diff --git a/Engine/Source/Developer/AutomationWindow/Private/SAutomationTestItem.cpp b/Engine/Source/Developer/AutomationWindow/Private/SAutomationTestItem.cpp
--- a/Engine/Source/Developer/AutomationWindow/Private/SAutomationTestItem.cpp
+++ b/Engine/Source/Developer/AutomationWindow/Private/SAutomationTestItem.cpp
@@ -109,7 +109,7 @@ TSharedRef<SWidget> SAutomationTestItem::GenerateWidgetForColumn( const FName& C
 #if WITH_EDITOR
 			TestNameWidget = SNew(SHyperlink)
 				.Style(FAutomationWindowStyle::Get(), "Common.GotoNativeCodeHyperlink")
-				.OnNavigate_Lambda([=] {
+				.OnNavigate_Lambda([this] {
 					GEngine->Exec(nullptr, *TestStatus->GetOpenCommand());
 				})
 				.Text(FText::FromString(TestStatus->GetDisplayNameWithDecoration()));
@@ -120,7 +120,7 @@ TSharedRef<SWidget> SAutomationTestItem::GenerateWidgetForColumn( const FName& C
 #if WITH_EDITOR
 			TestNameWidget = SNew(SHyperlink)
 				.Style(FAutomationWindowStyle::Get(), "Common.GotoNativeCodeHyperlink")
-				.OnNavigate_Lambda([=] {
+				.OnNavigate_Lambda([this] {
 					FString AssetPath = TestStatus->GetAssetPath();
 					FAssetRegistryModule& AssetRegistryModule = FModuleManager::LoadModuleChecked<FAssetRegistryModule>(TEXT("AssetRegistry"));
 
@@ -143,7 +143,7 @@ TSharedRef<SWidget> SAutomationTestItem::GenerateWidgetForColumn( const FName& C
 		{
 			TestNameWidget = SNew(SHyperlink)
 				.Style(FAutomationWindowStyle::Get(), "Common.GotoNativeCodeHyperlink")
-				.OnNavigate_Lambda([=] { FSlateApplication::Get().GotoLineInSource(TestStatus->GetSourceFile(), TestStatus->GetSourceFileLine()); })
+				.OnNavigate_Lambda([this] { FSlateApplication::Get().GotoLineInSource(TestStatus->GetSourceFile(), TestStatus->GetSourceFileLine()); })
 				.Text(FText::FromString(TestStatus->GetDisplayNameWithDecoration()));
 		}
 		else
diff --git a/Engine/Source/Developer/DeviceManager/Private/Widgets/Browser/SDeviceBrowserFilterBar.cpp b/Engine/Source/Developer/DeviceManager/Private/Widgets/Browser/SDeviceBrowserFilterBar.cpp
--- a/Engine/Source/Developer/DeviceManager/Private/Widgets/Browser/SDeviceBrowserFilterBar.cpp
+++ b/Engine/Source/Developer/DeviceManager/Private/Widgets/Browser/SDeviceBrowserFilterBar.cpp
@@ -57,7 +57,7 @@ void SDeviceBrowserFilterBar::Construct(const FArguments& InArgs, TSharedRef<FDe
 			[
 				SNew(SCheckBox)
 				.IsChecked_Lambda(
-					[=]() -> ECheckBoxState {
+					[this, PlatformEntry]() -> ECheckBoxState {
 						return Filter->IsPlatformEnabled(PlatformEntry->PlatformName)
 							? ECheckBoxState::Checked
 							: ECheckBoxState::Unchecked;
@@ -65,7 +65,7 @@ void SDeviceBrowserFilterBar::Construct(const FArguments& InArgs, TSharedRef<FDe
 				)
 				.Padding(FMargin(6.0, 2.0))
 				.OnCheckStateChanged_Lambda(
-					[=](ECheckBoxState CheckState) {
+					[this, PlatformEntry](ECheckBoxState CheckState) {
 						Filter->SetPlatformEnabled(PlatformEntry->PlatformName, CheckState == ECheckBoxState::Checked);
 					}
 				)
@@ -92,7 +92,7 @@ void SDeviceBrowserFilterBar::Construct(const FArguments& InArgs, TSharedRef<FDe
 						[
 							SNew(STextBlock)
 								.Text_Lambda(
-									[=]() -> FText {
+									[this, PlatformEntry]() -> FText {
 										return FText::Format(LOCTEXT("PlatformListRowFmt", "{0} ({1})"), FText::FromString(PlatformEntry->PlatformName), FText::AsNumber(Filter->GetServiceCountPerPlatform(PlatformEntry->PlatformName)));
 									}
 								)
diff --git a/Engine/Source/Developer/DirectoryWatcher/Private/DirectoryWatcherTests.cpp b/Engine/Source/Developer/DirectoryWatcher/Private/DirectoryWatcherTests.cpp
--- a/Engine/Source/Developer/DirectoryWatcher/Private/DirectoryWatcherTests.cpp
+++ b/Engine/Source/Developer/DirectoryWatcher/Private/DirectoryWatcherTests.cpp
@@ -94,7 +94,7 @@ bool FDirectoryWatcherSimpleCreateTest::RunTest(const FString& Parameters)
 	TSharedPtr<FDirectoryWatcherTestPayload> Test = MakeShareable(new FDirectoryWatcherTestPayload(WorkingDir));
 
 	// Give the stream time to start up before doing the test
-	AddCommand(new FDelayedFunctionLatentCommand([=]{
+	AddCommand(new FDelayedFunctionLatentCommand([this, WorkingDir, Test]{
 
 		// Create a file and check that it got reported as created
 		FFileHelper::SaveStringToFile(TEXT(""), *(WorkingDir / Filename));
@@ -125,7 +125,7 @@ bool FDirectoryWatcherSimpleModifyTest::RunTest(const FString& Parameters)
 	// Create a file first
 	FFileHelper::SaveStringToFile(TEXT(""), *(WorkingDir / Filename));
 
-	AddCommand(new FDelayedFunctionLatentCommand([=]{
+	AddCommand(new FDelayedFunctionLatentCommand([this, WorkingDir]{
 
 		// Start watching the directory
 		TSharedPtr<FDirectoryWatcherTestPayload> Test = MakeShareable(new FDirectoryWatcherTestPayload(WorkingDir));
@@ -167,7 +167,7 @@ bool FDirectoryWatcherSimpleDeleteTest::RunTest(const FString& Parameters)
 	TSharedPtr<FDirectoryWatcherTestPayload> Test = MakeShareable(new FDirectoryWatcherTestPayload(WorkingDir));
 
 	// Give the stream time to start up before doing the test
-	AddCommand(new FDelayedFunctionLatentCommand([=]{
+	AddCommand(new FDelayedFunctionLatentCommand([this, WorkingDir, Test]{
 
 		// Delete the file
 		IFileManager::Get().Delete(*(WorkingDir / Filename));
@@ -203,7 +203,7 @@ bool FDirectoryWatcherSubFolderTest::RunTest(const FString& Parameters)
 	FFileHelper::SaveStringToFile(TEXT(""), *(WorkingDir / RemovedFilename));
 
 	// Give the stream time to start up before doing the test
-	AddCommand(new FDelayedFunctionLatentCommand([=]{
+	AddCommand(new FDelayedFunctionLatentCommand([this, WorkingDir]{
 
 		// Start watching the directory
 		TSharedPtr<FDirectoryWatcherTestPayload> Test = MakeShareable(new FDirectoryWatcherTestPayload(WorkingDir));
@@ -258,7 +258,7 @@ bool FDirectoryWatcherNewFolderTest::RunTest(const FString& Parameters)
 	IFileManager::Get().DeleteDirectory(*( WorkingDir / CreatedDirectory ), true);
 
 	// Give the stream time to start up before doing the test
-	AddCommand(new FDelayedFunctionLatentCommand([=] {
+	AddCommand(new FDelayedFunctionLatentCommand([this, WorkingDir] {
 
 		IFileManager::Get().MakeDirectory(*(WorkingDir / RemovedDirectory), true);
 
@@ -305,7 +305,7 @@ bool FDirectoryWatcherIgnoreSubtreeTest::RunTest(const FString& Parameters)
 	IFileManager::Get().DeleteDirectory(*( WorkingDir / ChildDirectory ), true);
 
 	// Give the stream time to start up before doing the test
-	AddCommand(new FDelayedFunctionLatentCommand([=]{
+	AddCommand(new FDelayedFunctionLatentCommand([this, WorkingDir]{
 
 		// Start watching the directory
 		TSharedPtr<FDirectoryWatcherTestPayload> Test = MakeShareable(new FDirectoryWatcherTestPayload(WorkingDir, IDirectoryWatcher::WatchOptions::IgnoreChangesInSubtree | IDirectoryWatcher::WatchOptions::IncludeDirectoryChanges));
diff --git a/Engine/Source/Editor/ClothPainter/Private/ClothPaintingModule.cpp b/Engine/Source/Editor/ClothPainter/Private/ClothPaintingModule.cpp
--- a/Engine/Source/Editor/ClothPainter/Private/ClothPaintingModule.cpp
+++ b/Engine/Source/Editor/ClothPainter/Private/ClothPaintingModule.cpp
@@ -250,7 +250,7 @@ void FClothPaintingModule::ShutdownModule()
 	{
 		TArray<ISkeletalMeshEditorModule::FSkeletalMeshEditorToolbarExtender>& Extenders = SkelMeshEditorModule->GetAllSkeletalMeshEditorToolbarExtenders();
 
-		Extenders.RemoveAll([=](const ISkeletalMeshEditorModule::FSkeletalMeshEditorToolbarExtender& InDelegate) {return InDelegate.GetHandle() == SkelMeshEditorExtenderHandle; });
+		Extenders.RemoveAll([this](const ISkeletalMeshEditorModule::FSkeletalMeshEditorToolbarExtender& InDelegate) {return InDelegate.GetHandle() == SkelMeshEditorExtenderHandle; });
 	}
 }
 
diff --git a/Engine/Source/Editor/DetailCustomizations/Private/LandscapeUIDetails.cpp b/Engine/Source/Editor/DetailCustomizations/Private/LandscapeUIDetails.cpp
--- a/Engine/Source/Editor/DetailCustomizations/Private/LandscapeUIDetails.cpp
+++ b/Engine/Source/Editor/DetailCustomizations/Private/LandscapeUIDetails.cpp
@@ -96,7 +96,7 @@ void FLandscapeUIDetails::CustomizeDetails( IDetailLayoutBuilder& DetailBuilder
 			{
 				return Landscape.IsValid() && Landscape->CanHaveLayersContent() ? ECheckBoxState::Checked : ECheckBoxState::Unchecked;
 			})
-			.OnCheckStateChanged_Lambda([=](ECheckBoxState NewState)
+			.OnCheckStateChanged_Lambda([this, Landscape](ECheckBoxState NewState)
 			{
 				bool bChecked = (NewState == ECheckBoxState::Checked);
 				if (Landscape.IsValid() && (Landscape->CanHaveLayersContent() != bChecked))
diff --git a/Engine/Source/Editor/DetailCustomizations/Private/MaterialAttributePropertyDetails.cpp b/Engine/Source/Editor/DetailCustomizations/Private/MaterialAttributePropertyDetails.cpp
--- a/Engine/Source/Editor/DetailCustomizations/Private/MaterialAttributePropertyDetails.cpp
+++ b/Engine/Source/Editor/DetailCustomizations/Private/MaterialAttributePropertyDetails.cpp
@@ -79,11 +79,11 @@ void FMaterialAttributePropertyDetails::OnBuildChild(TSharedRef<IPropertyHandle>
 
 	FPropertyComboBoxArgs ComboArgs(
 		ChildHandle,
-		FOnGetPropertyComboBoxStrings::CreateLambda([=](TArray<TSharedPtr<FString>>& OutComboBoxStrings, TArray<TSharedPtr<SToolTip>>& OutToolTips, TArray<bool>& OutRestrictedItems) -> void  
+		FOnGetPropertyComboBoxStrings::CreateLambda([this](TArray<TSharedPtr<FString>>& OutComboBoxStrings, TArray<TSharedPtr<SToolTip>>& OutToolTips, TArray<bool>& OutRestrictedItems) -> void  
 		{									
 			OutComboBoxStrings = AttributeDisplayNameList;
 		}), 
-		FOnGetPropertyComboBoxValue::CreateLambda( [=]() -> FString
+		FOnGetPropertyComboBoxValue::CreateLambda( [this, ChildHandle]() -> FString
 		{
 			FString AttributeName;
 			if (ChildHandle->IsValidHandle())
@@ -97,7 +97,7 @@ void FMaterialAttributePropertyDetails::OnBuildChild(TSharedRef<IPropertyHandle>
 			} 
 			return AttributeName;
 		}),
-		FOnPropertyComboBoxValueSelected::CreateLambda( [=] (const FString& Selection)
+		FOnPropertyComboBoxValueSelected::CreateLambda( [this, ChildHandle] (const FString& Selection)
 		{
 			if (ChildHandle->IsValidHandle())
 			{
diff --git a/Engine/Source/Editor/FoliageEdit/Private/FoliageEdMode.cpp b/Engine/Source/Editor/FoliageEdit/Private/FoliageEdMode.cpp
--- a/Engine/Source/Editor/FoliageEdit/Private/FoliageEdMode.cpp
+++ b/Engine/Source/Editor/FoliageEdit/Private/FoliageEdMode.cpp
@@ -367,7 +367,7 @@ void FEdModeFoliage::BindCommands()
 		Commands.SetPaint,
 		FExecuteAction::CreateRaw(this, &FEdModeFoliage::OnSetPaint),
 		FCanExecuteAction(),
-		FIsActionChecked::CreateLambda([=]
+		FIsActionChecked::CreateLambda([this]
 	{
 		return UISettings.GetPaintToolSelected() && !UISettings.GetIsInSingleInstantiationMode();
 	}));
@@ -376,7 +376,7 @@ void FEdModeFoliage::BindCommands()
 		Commands.SetReapplySettings,
 		FExecuteAction::CreateRaw(this, &FEdModeFoliage::OnSetReapplySettings),
 		FCanExecuteAction(),
-		FIsActionChecked::CreateLambda([=]
+		FIsActionChecked::CreateLambda([this]
 	{
 		return UISettings.GetReapplyToolSelected() && !UISettings.GetIsInSingleInstantiationMode();
 	}));
@@ -385,7 +385,7 @@ void FEdModeFoliage::BindCommands()
 		Commands.SetSelect,
 		FExecuteAction::CreateRaw(this, &FEdModeFoliage::OnSetSelectInstance),
 		FCanExecuteAction(),
-		FIsActionChecked::CreateLambda([=]
+		FIsActionChecked::CreateLambda([this]
 	{
 		return UISettings.GetSelectToolSelected();
 	}));
@@ -394,7 +394,7 @@ void FEdModeFoliage::BindCommands()
 		Commands.SetLassoSelect,
 		FExecuteAction::CreateRaw(this, &FEdModeFoliage::OnSetLasso),
 		FCanExecuteAction(),
-		FIsActionChecked::CreateLambda([=]
+		FIsActionChecked::CreateLambda([this]
 	{
 		return UISettings.GetLassoSelectToolSelected();
 	}));
@@ -403,7 +403,7 @@ void FEdModeFoliage::BindCommands()
 		Commands.SetPaintBucket,
 		FExecuteAction::CreateRaw(this, &FEdModeFoliage::OnSetPaintFill),
 		FCanExecuteAction(),
-		FIsActionChecked::CreateLambda([=]
+		FIsActionChecked::CreateLambda([this]
 	{
 		return UISettings.GetPaintBucketToolSelected();
 	}));
diff --git a/Engine/Source/Editor/FoliageEdit/Private/FoliageTypePaintingCustomization.cpp b/Engine/Source/Editor/FoliageEdit/Private/FoliageTypePaintingCustomization.cpp
--- a/Engine/Source/Editor/FoliageEdit/Private/FoliageTypePaintingCustomization.cpp
+++ b/Engine/Source/Editor/FoliageEdit/Private/FoliageTypePaintingCustomization.cpp
@@ -223,7 +223,7 @@ void FFoliageTypePaintingCustomization::GetHiddenPropertyVisibility(const TShare
 	if (bHideInReapplyTool)
 	{
 		// In addition to hiding it it behind the given property, only show this in the reapply tool
-		OutVisibility = TAttribute<EVisibility>::Create(TAttribute<EVisibility>::FGetter::CreateLambda([=]
+		OutVisibility = TAttribute<EVisibility>::Create(TAttribute<EVisibility>::FGetter::CreateLambda([this, VisibilityGetter]
 		{
 			if (!FoliageEditMode->UISettings.GetReapplyToolSelected() && VisibilityGetter.IsBound())
 			{
diff --git a/Engine/Source/Editor/FoliageEdit/Private/SFoliageEdit.cpp b/Engine/Source/Editor/FoliageEdit/Private/SFoliageEdit.cpp
--- a/Engine/Source/Editor/FoliageEdit/Private/SFoliageEdit.cpp
+++ b/Engine/Source/Editor/FoliageEdit/Private/SFoliageEdit.cpp
@@ -610,7 +610,7 @@ TSharedRef<SWidget> SFoliageEdit::MakeFilterMenu()
 		LOCTEXT("FilterLandscapeTooltip", "Allow Foliage to be placed on Landscape"),
 		FSlateIcon(),
 		FUIAction(
-			FExecuteAction::CreateLambda( [=] { FoliageEditMode->UISettings.SetFilterLandscape( !FoliageEditMode->UISettings.GetFilterLandscape()); }),
+			FExecuteAction::CreateLambda( [this] { FoliageEditMode->UISettings.SetFilterLandscape( !FoliageEditMode->UISettings.GetFilterLandscape()); }),
 			FCanExecuteAction(),
 			FIsActionChecked::CreateRaw( &FoliageEditMode->UISettings, &FFoliageUISettings::GetFilterLandscape)
 		),
@@ -624,7 +624,7 @@ TSharedRef<SWidget> SFoliageEdit::MakeFilterMenu()
 		LOCTEXT("FilterStaticMeshTooltip", "Allow Foliage to be placed on StaticMesh"),
 		FSlateIcon(),
 		FUIAction(
-			FExecuteAction::CreateLambda( [=] { FoliageEditMode->UISettings.SetFilterStaticMesh( !FoliageEditMode->UISettings.GetFilterStaticMesh()); }),
+			FExecuteAction::CreateLambda( [this] { FoliageEditMode->UISettings.SetFilterStaticMesh( !FoliageEditMode->UISettings.GetFilterStaticMesh()); }),
 			FCanExecuteAction(),
 			FIsActionChecked::CreateRaw( &FoliageEditMode->UISettings, &FFoliageUISettings::GetFilterStaticMesh)
 		),
@@ -638,7 +638,7 @@ TSharedRef<SWidget> SFoliageEdit::MakeFilterMenu()
 		LOCTEXT("FilterBSPTooltip", "Allow Foliage to be placed on BSP"),
 		FSlateIcon(),
 		FUIAction(
-			FExecuteAction::CreateLambda( [=] { FoliageEditMode->UISettings.SetFilterBSP( !FoliageEditMode->UISettings.GetFilterBSP()); }),
+			FExecuteAction::CreateLambda( [this] { FoliageEditMode->UISettings.SetFilterBSP( !FoliageEditMode->UISettings.GetFilterBSP()); }),
 			FCanExecuteAction(),
 			FIsActionChecked::CreateRaw( &FoliageEditMode->UISettings, &FFoliageUISettings::GetFilterBSP)
 		),
@@ -652,7 +652,7 @@ TSharedRef<SWidget> SFoliageEdit::MakeFilterMenu()
 		LOCTEXT("FilterFoliageTooltip", "Allow Foliage to be placed on Foliage"),
 		FSlateIcon(),
 		FUIAction(
-			FExecuteAction::CreateLambda( [=] { FoliageEditMode->UISettings.SetFilterFoliage( !FoliageEditMode->UISettings.GetFilterFoliage()); }),
+			FExecuteAction::CreateLambda( [this] { FoliageEditMode->UISettings.SetFilterFoliage( !FoliageEditMode->UISettings.GetFilterFoliage()); }),
 			FCanExecuteAction(),
 			FIsActionChecked::CreateRaw( &FoliageEditMode->UISettings, &FFoliageUISettings::GetFilterFoliage)
 		),
@@ -666,7 +666,7 @@ TSharedRef<SWidget> SFoliageEdit::MakeFilterMenu()
 		LOCTEXT("FilterTranslucentTooltip", "Allow Foliage to be placed on Translucent"),
 		FSlateIcon(),
 		FUIAction(
-			FExecuteAction::CreateLambda( [=] { FoliageEditMode->UISettings.SetFilterTranslucent( !FoliageEditMode->UISettings.GetFilterTranslucent()); }),
+			FExecuteAction::CreateLambda( [this] { FoliageEditMode->UISettings.SetFilterTranslucent( !FoliageEditMode->UISettings.GetFilterTranslucent()); }),
 			FCanExecuteAction(),
 			FIsActionChecked::CreateRaw( &FoliageEditMode->UISettings, &FFoliageUISettings::GetFilterTranslucent)
 		),
@@ -688,7 +688,7 @@ TSharedRef<SWidget> SFoliageEdit::MakeSettingsMenu()
 		LOCTEXT("SettingsCurrentLevelTooltip", "Allow Foliage to be placed on Translucent"),
 		FSlateIcon(),
 		FUIAction(
-			FExecuteAction::CreateLambda( [=] { 
+			FExecuteAction::CreateLambda( [this] { 
 				FoliageEditMode->UISettings.SetSpawnInCurrentLevelMode( !FoliageEditMode->UISettings.GetIsInSpawnInCurrentLevelMode() ); 
 			}),
 			FCanExecuteAction(),
@@ -1290,7 +1290,7 @@ TSharedRef<SWidget> SFoliageEdit::GetSingleInstantiationModeMenuContent()
 			FUIAction(
 				FExecuteAction::CreateSP(this, &SFoliageEdit::OnSingleInstantiationPlacementModeChanged, i),
 				FCanExecuteAction(),
-				FIsActionChecked::CreateLambda([=] { return FoliageEditMode->UISettings.GetSingleInstantiationPlacementMode() == EFoliageSingleInstantiationPlacementMode::Type(i); })
+				FIsActionChecked::CreateLambda([this, i] { return FoliageEditMode->UISettings.GetSingleInstantiationPlacementMode() == EFoliageSingleInstantiationPlacementMode::Type(i); })
 			),
 			NAME_None,
 			EUserInterfaceActionType::ToggleButton
diff --git a/Engine/Source/Editor/LandscapeEditor/Private/LandscapeEditorDetailCustomization_Layers.cpp b/Engine/Source/Editor/LandscapeEditor/Private/LandscapeEditorDetailCustomization_Layers.cpp
--- a/Engine/Source/Editor/LandscapeEditor/Private/LandscapeEditorDetailCustomization_Layers.cpp
+++ b/Engine/Source/Editor/LandscapeEditor/Private/LandscapeEditorDetailCustomization_Layers.cpp
@@ -266,14 +266,14 @@ TSharedPtr<SWidget> FLandscapeEditorCustomNodeBuilder_Layers::GenerateRow(int32
 					.IsEnabled(this, &FLandscapeEditorCustomNodeBuilder_Layers::IsLayerEditionEnabled, InLayerIndex)
 					.Visibility(this, &FLandscapeEditorCustomNodeBuilder_Layers::GetLayerAlphaVisibility, InLayerIndex)
 					.Value(this, &FLandscapeEditorCustomNodeBuilder_Layers::GetLayerAlpha, InLayerIndex)
-					.OnValueChanged_Lambda([=](float InValue) { SetLayerAlpha(InValue, InLayerIndex, false); })
-					.OnValueCommitted_Lambda([=](float InValue, ETextCommit::Type InCommitType) { SetLayerAlpha(InValue, InLayerIndex, true); })
-					.OnBeginSliderMovement_Lambda([=]()
+					.OnValueChanged_Lambda([this, InLayerIndex](float InValue) { SetLayerAlpha(InValue, InLayerIndex, false); })
+					.OnValueCommitted_Lambda([this, InLayerIndex](float InValue, ETextCommit::Type InCommitType) { SetLayerAlpha(InValue, InLayerIndex, true); })
+					.OnBeginSliderMovement_Lambda([this, InLayerIndex]()
 					{
 						CurrentSlider = InLayerIndex;
 						GEditor->BeginTransaction(LOCTEXT("Landscape_Layers_SetAlpha", "Set Layer Alpha"));
 					})
-					.OnEndSliderMovement_Lambda([=](double)
+					.OnEndSliderMovement_Lambda([this](double)
 					{
 						GEditor->EndTransaction();
 						CurrentSlider = INDEX_NONE;
diff --git a/Engine/Source/Editor/LandscapeEditor/Private/LandscapeEditorDetailCustomization_NewLandscape.cpp b/Engine/Source/Editor/LandscapeEditor/Private/LandscapeEditorDetailCustomization_NewLandscape.cpp
--- a/Engine/Source/Editor/LandscapeEditor/Private/LandscapeEditorDetailCustomization_NewLandscape.cpp
+++ b/Engine/Source/Editor/LandscapeEditor/Private/LandscapeEditorDetailCustomization_NewLandscape.cpp
@@ -480,20 +480,20 @@ void FLandscapeEditorDetailCustomization_NewLandscape::CustomizeDetails(IDetailL
 			//.MinSliderValue(TAttribute<TOptional<int32> >(this, &FLandscapeEditorDetailCustomization_NewLandscape::GetMinLandscapeResolution))
 			//.MaxSliderValue(TAttribute<TOptional<int32> >(this, &FLandscapeEditorDetailCustomization_NewLandscape::GetMaxLandscapeResolution))
 			.Value(this, &FLandscapeEditorDetailCustomization_NewLandscape::GetLandscapeResolutionX)
-			.OnValueChanged_Lambda([=](int32 NewValue)
+			.OnValueChanged_Lambda([this](int32 NewValue)
 			{
 				OnChangeLandscapeResolutionX(NewValue, false);
 			})
-			.OnValueCommitted_Lambda([=](int32 NewValue, ETextCommit::Type)
+			.OnValueCommitted_Lambda([this](int32 NewValue, ETextCommit::Type)
 			{
 				OnChangeLandscapeResolutionX(NewValue, true);
 			})
-			.OnBeginSliderMovement_Lambda([=]()
+			.OnBeginSliderMovement_Lambda([this]()
 			{
 				bUsingSlider = true;
 				GEditor->BeginTransaction(LOCTEXT("ChangeResolutionX_Transaction", "Change Landscape Resolution X"));
 			})
-			.OnEndSliderMovement_Lambda([=](double)
+			.OnEndSliderMovement_Lambda([this](double)
 			{
 				GEditor->EndTransaction();
 				bUsingSlider = false;
@@ -522,20 +522,20 @@ void FLandscapeEditorDetailCustomization_NewLandscape::CustomizeDetails(IDetailL
 			//.MinSliderValue(TAttribute<TOptional<int32> >(this, &FLandscapeEditorDetailCustomization_NewLandscape::GetMinLandscapeResolution))
 			//.MaxSliderValue(TAttribute<TOptional<int32> >(this, &FLandscapeEditorDetailCustomization_NewLandscape::GetMaxLandscapeResolution))
 			.Value(this, &FLandscapeEditorDetailCustomization_NewLandscape::GetLandscapeResolutionY)
-			.OnValueChanged_Lambda([=](int32 NewValue)
+			.OnValueChanged_Lambda([this](int32 NewValue)
 			{
 				OnChangeLandscapeResolutionY(NewValue, false);
 			})
-			.OnValueCommitted_Lambda([=](int32 NewValue, ETextCommit::Type)
+			.OnValueCommitted_Lambda([this](int32 NewValue, ETextCommit::Type)
 			{
 				OnChangeLandscapeResolutionY(NewValue, true);
 			})
-			.OnBeginSliderMovement_Lambda([=]()
+			.OnBeginSliderMovement_Lambda([this]()
 			{
 				bUsingSlider = true;
 				GEditor->BeginTransaction(LOCTEXT("ChangeResolutionY_Transaction", "Change Landscape Resolution Y"));
 			})
-			.OnEndSliderMovement_Lambda([=](double)
+			.OnEndSliderMovement_Lambda([this](double)
 			{
 				GEditor->EndTransaction();
 				bUsingSlider = false;
diff --git a/Engine/Source/Editor/MergeActors/Private/MergeProxyUtils/Utils.cpp b/Engine/Source/Editor/MergeActors/Private/MergeProxyUtils/Utils.cpp
--- a/Engine/Source/Editor/MergeActors/Private/MergeProxyUtils/Utils.cpp
+++ b/Engine/Source/Editor/MergeActors/Private/MergeProxyUtils/Utils.cpp
@@ -206,7 +206,7 @@ TSharedRef<ITableRow> FComponentSelectionControl::MakeComponentListItemWidget(TS
 					.IsChecked(State)
 					.ToolTipText(LOCTEXT("IncorporateCheckBoxToolTip", "When ticked the Component will be incorporated into the merge"))
 
-					.OnCheckStateChanged_Lambda([=](ECheckBoxState NewState)
+					.OnCheckStateChanged_Lambda([this, ComponentData, bIsMesh](ECheckBoxState NewState)
 					{
 						ComponentData->bShouldIncorporate = (NewState == ECheckBoxState::Checked);
 
diff --git a/Engine/Source/Editor/MeshPaint/Private/IMeshPaintMode.cpp b/Engine/Source/Editor/MeshPaint/Private/IMeshPaintMode.cpp
--- a/Engine/Source/Editor/MeshPaint/Private/IMeshPaintMode.cpp
+++ b/Engine/Source/Editor/MeshPaint/Private/IMeshPaintMode.cpp
@@ -86,7 +86,7 @@ void IMeshPaintEdMode::Enter()
 	// Initialize adapter globals
 	FMeshPaintAdapterFactory::InitializeAdapterGlobals();
 
-	SelectionChangedHandle = USelection::SelectionChangedEvent.AddLambda([=](UObject* Object) { MeshPainter->Refresh();  });
+	SelectionChangedHandle = USelection::SelectionChangedEvent.AddLambda([this](UObject* Object) { MeshPainter->Refresh();  });
 
 	if (UsesToolkits() && !Toolkit.IsValid())
 	{
diff --git a/Engine/Source/Editor/MeshPaint/Private/MeshPaintSkeletalMeshAdapter.cpp b/Engine/Source/Editor/MeshPaint/Private/MeshPaintSkeletalMeshAdapter.cpp
--- a/Engine/Source/Editor/MeshPaint/Private/MeshPaintSkeletalMeshAdapter.cpp
+++ b/Engine/Source/Editor/MeshPaint/Private/MeshPaintSkeletalMeshAdapter.cpp
@@ -218,7 +218,7 @@ void FMeshPaintGeometryAdapterForSkeletalMeshes::OnAdded()
 	FSkeletalMeshReferencers& SkeletalMeshReferencers = MeshToComponentMap.FindOrAdd(ReferencedSkeletalMesh);
 
 	checkf(!SkeletalMeshReferencers.Referencers.ContainsByPredicate(
-		[=](const FSkeletalMeshReferencers::FReferencersInfo& Info)
+		[this](const FSkeletalMeshReferencers::FReferencersInfo& Info)
 		{
 			return Info.SkeletalMeshComponent == this->SkeletalMeshComponent;
 		}), TEXT("This Skeletal Mesh Component has already been Added"));
@@ -286,7 +286,7 @@ void FMeshPaintGeometryAdapterForSkeletalMeshes::OnRemoved()
 	checkf(SkeletalMeshReferencers->Referencers.Num() > 0, TEXT("Skeletal Mesh does not have any referencers"));
 
 	const int32 Index = SkeletalMeshReferencers->Referencers.IndexOfByPredicate(
-		[=](const FSkeletalMeshReferencers::FReferencersInfo& Info)
+		[this](const FSkeletalMeshReferencers::FReferencersInfo& Info)
 		{
 			return Info.SkeletalMeshComponent == this->SkeletalMeshComponent;
 		}
diff --git a/Engine/Source/Editor/MeshPaint/Private/MeshPaintStaticMeshAdapter.cpp b/Engine/Source/Editor/MeshPaint/Private/MeshPaintStaticMeshAdapter.cpp
--- a/Engine/Source/Editor/MeshPaint/Private/MeshPaintStaticMeshAdapter.cpp
+++ b/Engine/Source/Editor/MeshPaint/Private/MeshPaintStaticMeshAdapter.cpp
@@ -166,7 +166,7 @@ void FMeshPaintGeometryAdapterForStaticMeshes::OnAdded()
 	FStaticMeshReferencers& StaticMeshReferencers = MeshToComponentMap.FindOrAdd(ReferencedStaticMesh);
 
 	check(!StaticMeshReferencers.Referencers.ContainsByPredicate(
-		[=](const FStaticMeshReferencers::FReferencersInfo& Info)
+		[this](const FStaticMeshReferencers::FReferencersInfo& Info)
 	{
 		return Info.StaticMeshComponent == this->StaticMeshComponent;
 	}
@@ -233,7 +233,7 @@ void FMeshPaintGeometryAdapterForStaticMeshes::OnRemoved()
 	{
 		check(StaticMeshReferencers->Referencers.Num() > 0);
 		int32 Index = StaticMeshReferencers->Referencers.IndexOfByPredicate(
-			[=](const FStaticMeshReferencers::FReferencersInfo& Info)
+			[this](const FStaticMeshReferencers::FReferencersInfo& Info)
 		{
 			return Info.StaticMeshComponent == this->StaticMeshComponent;
 		}
diff --git a/Engine/Source/Editor/MovieSceneTools/Private/TrackEditors/SkeletalAnimationTrackEditor.cpp b/Engine/Source/Editor/MovieSceneTools/Private/TrackEditors/SkeletalAnimationTrackEditor.cpp
--- a/Engine/Source/Editor/MovieSceneTools/Private/TrackEditors/SkeletalAnimationTrackEditor.cpp
+++ b/Engine/Source/Editor/MovieSceneTools/Private/TrackEditors/SkeletalAnimationTrackEditor.cpp
@@ -785,7 +785,7 @@ void FSkeletalAnimationSection::BuildSectionContextMenu(FMenuBuilder& MenuBuilde
 			BoneNames.Add(Skeleton->GetReferenceSkeleton().GetBoneName(BoneIndex));
 		}
 
-		auto MatchToBone = [=](bool bMatchPrevious, int32 Index)
+		auto MatchToBone = [=, this](bool bMatchPrevious, int32 Index)
 		{
 			return FUIAction(
 				FExecuteAction::CreateLambda([=, this]
@@ -925,7 +925,7 @@ void FSkeletalAnimationSection::BuildSectionContextMenu(FMenuBuilder& MenuBuilde
 				LOCTEXT("MatchPitchRotationTooltip", "Match the Pitch Rotation, may want this off for better matching"),
 				FSlateIcon(),
 				FUIAction(
-					FExecuteAction::CreateLambda([=]()->void { 
+					FExecuteAction::CreateLambda([=, this]()->void { 
 						FScopedTransaction MatchTransaction(LOCTEXT("MatchPitchRotation_Transaction", "Match Pitch Rotation"));
 						Section.Modify();
 						Section.ToggleMatchIncludePitchRotation();
diff --git a/Engine/Source/Editor/Persona/Private/AnimTimeline/AnimTimeSliderController.cpp b/Engine/Source/Editor/Persona/Private/AnimTimeline/AnimTimeSliderController.cpp
--- a/Engine/Source/Editor/Persona/Private/AnimTimeline/AnimTimeSliderController.cpp
+++ b/Engine/Source/Editor/Persona/Private/AnimTimeline/AnimTimeSliderController.cpp
@@ -991,8 +991,8 @@ TSharedRef<SWidget> FAnimTimeSliderController::OpenSetPlaybackRangeMenu(FFrameNu
 			FText(),
 			FSlateIcon(),
 			FUIAction(
-				FExecuteAction::CreateLambda([=]{ SetSelectionRangeStart(FrameNumber); }),
-				FCanExecuteAction::CreateLambda([=]{ return SelectionRange.IsEmpty() || FrameNumber < UE::MovieScene::DiscreteExclusiveUpper(SelectionRange); })
+				FExecuteAction::CreateLambda([this, FrameNumber]{ SetSelectionRangeStart(FrameNumber); }),
+				FCanExecuteAction::CreateLambda([FrameNumber, SelectionRange]{ return SelectionRange.IsEmpty() || FrameNumber < UE::MovieScene::DiscreteExclusiveUpper(SelectionRange); })
 			)
 		);
 
@@ -1001,8 +1001,8 @@ TSharedRef<SWidget> FAnimTimeSliderController::OpenSetPlaybackRangeMenu(FFrameNu
 			FText(),
 			FSlateIcon(),
 			FUIAction(
-				FExecuteAction::CreateLambda([=]{ SetSelectionRangeEnd(FrameNumber); }),
-				FCanExecuteAction::CreateLambda([=]{ return SelectionRange.IsEmpty() || FrameNumber >= UE::MovieScene::DiscreteInclusiveLower(SelectionRange); })
+				FExecuteAction::CreateLambda([this, FrameNumber]{ SetSelectionRangeEnd(FrameNumber); }),
+				FCanExecuteAction::CreateLambda([FrameNumber, SelectionRange]{ return SelectionRange.IsEmpty() || FrameNumber >= UE::MovieScene::DiscreteInclusiveLower(SelectionRange); })
 			)
 		);
 
@@ -1011,8 +1011,8 @@ TSharedRef<SWidget> FAnimTimeSliderController::OpenSetPlaybackRangeMenu(FFrameNu
 			FText(),
 			FSlateIcon(),
 			FUIAction(
-				FExecuteAction::CreateLambda([=]{ TimeSliderArgs.OnSelectionRangeChanged.ExecuteIfBound(TRange<FFrameNumber>::Empty()); }),
-				FCanExecuteAction::CreateLambda([=]{ return !SelectionRange.IsEmpty(); })
+				FExecuteAction::CreateLambda([this, FrameNumber]{ TimeSliderArgs.OnSelectionRangeChanged.ExecuteIfBound(TRange<FFrameNumber>::Empty()); }),
+				FCanExecuteAction::CreateLambda([FrameNumber, SelectionRange]{ return !SelectionRange.IsEmpty(); })
 			)
 		);
 	}
diff --git a/Engine/Source/Editor/WorldBrowser/Private/SWorldHierarchyImpl.cpp b/Engine/Source/Editor/WorldBrowser/Private/SWorldHierarchyImpl.cpp
--- a/Engine/Source/Editor/WorldBrowser/Private/SWorldHierarchyImpl.cpp
+++ b/Engine/Source/Editor/WorldBrowser/Private/SWorldHierarchyImpl.cpp
@@ -472,7 +472,7 @@ TSharedPtr<SWidget> SWorldHierarchyImpl::ConstructLevelContextMenu() const
 			RootTreeItems[0]->GenerateContextMenu(Menu, *this);
 		}
 
-		Menu->AddDynamicSection("HierarchyDynamicSection", FNewToolMenuDelegateLegacy::CreateLambda([=](FMenuBuilder& MenuBuilder, UToolMenu* Menu)
+		Menu->AddDynamicSection("HierarchyDynamicSection", FNewToolMenuDelegateLegacy::CreateLambda([this](FMenuBuilder& MenuBuilder, UToolMenu* Menu)
 		{
 			WorldModel->BuildHierarchyMenu(MenuBuilder);
 		}));
diff --git a/Engine/Source/Editor/WorldPartitionEditor/Private/WorldPartition/SWorldPartitionEditorGrid2D.cpp b/Engine/Source/Editor/WorldPartitionEditor/Private/WorldPartition/SWorldPartitionEditorGrid2D.cpp
--- a/Engine/Source/Editor/WorldPartitionEditor/Private/WorldPartition/SWorldPartitionEditorGrid2D.cpp
+++ b/Engine/Source/Editor/WorldPartitionEditor/Private/WorldPartition/SWorldPartitionEditorGrid2D.cpp
@@ -314,7 +314,7 @@ void SWorldPartitionEditorGrid2D::Construct(const FArguments& InArgs)
 				.ForegroundColor(FSlateColor::UseForeground())
 				.BorderBackgroundColor(TextBoxStyle.BackgroundColor)
 				.IsChecked(bShowActors ? ECheckBoxState::Checked : ECheckBoxState::Unchecked)
-				.OnCheckStateChanged(FOnCheckStateChanged::CreateLambda([=](ECheckBoxState State) { bShowActors = !bShowActors; }))
+				.OnCheckStateChanged(FOnCheckStateChanged::CreateLambda([this](ECheckBoxState State) { bShowActors = !bShowActors; }))
 				[
 					SNew(SBox)
 					.VAlign(VAlign_Center)
@@ -330,7 +330,7 @@ void SWorldPartitionEditorGrid2D::Construct(const FArguments& InArgs)
 				.ForegroundColor(FSlateColor::UseForeground())
 				.BorderBackgroundColor(TextBoxStyle.BackgroundColor)
 				.IsChecked(bShowGrid ? ECheckBoxState::Checked : ECheckBoxState::Unchecked)
-				.OnCheckStateChanged(FOnCheckStateChanged::CreateLambda([=](ECheckBoxState State) { bShowGrid = !bShowGrid; }))
+				.OnCheckStateChanged(FOnCheckStateChanged::CreateLambda([this](ECheckBoxState State) { bShowGrid = !bShowGrid; }))
 				[
 					SNew(SBox)
 					.VAlign(VAlign_Center)
@@ -346,7 +346,7 @@ void SWorldPartitionEditorGrid2D::Construct(const FArguments& InArgs)
 				.ForegroundColor(FSlateColor::UseForeground())
 				.IsChecked(bFollowPlayerInPIE ? ECheckBoxState::Checked : ECheckBoxState::Unchecked)
 				.Visibility_Lambda([this]() { return GetDefault<UWorldPartitionEditorSettings>()->bDisablePIE ? EVisibility::Hidden : EVisibility::Visible; })
-				.OnCheckStateChanged(FOnCheckStateChanged::CreateLambda([=](ECheckBoxState State) { bFollowPlayerInPIE = !bFollowPlayerInPIE; }))
+				.OnCheckStateChanged(FOnCheckStateChanged::CreateLambda([this](ECheckBoxState State) { bFollowPlayerInPIE = !bFollowPlayerInPIE; }))
 				[
 					SNew(SBox)
 					.VAlign(VAlign_Center)
diff --git a/Engine/Source/Editor/WorldPartitionEditor/Private/WorldPartitionEditorModule.cpp b/Engine/Source/Editor/WorldPartitionEditor/Private/WorldPartitionEditorModule.cpp
--- a/Engine/Source/Editor/WorldPartitionEditor/Private/WorldPartitionEditorModule.cpp
+++ b/Engine/Source/Editor/WorldPartitionEditor/Private/WorldPartitionEditorModule.cpp
@@ -159,7 +159,7 @@ void FWorldPartitionEditorModule::ShutdownModule()
 	{
 		if (FLevelEditorModule* LevelEditorModule = FModuleManager::Get().GetModulePtr<FLevelEditorModule>("LevelEditor"))
 		{
-			LevelEditorModule->GetAllLevelViewportContextMenuExtenders().RemoveAll([=](const FLevelEditorModule::FLevelViewportMenuExtender_SelectedActors& In) { return In.GetHandle() == LevelEditorExtenderDelegateHandle; });
+			LevelEditorModule->GetAllLevelViewportContextMenuExtenders().RemoveAll([this](const FLevelEditorModule::FLevelViewportMenuExtender_SelectedActors& In) { return In.GetHandle() == LevelEditorExtenderDelegateHandle; });
 
 			LevelEditorModule->OnRegisterTabs().RemoveAll(this);
 			LevelEditorModule->OnRegisterLayoutExtensions().RemoveAll(this);
diff --git a/Engine/Source/Runtime/D3D12RHI/Private/Windows/WindowsD3D12Viewport.cpp b/Engine/Source/Runtime/D3D12RHI/Private/Windows/WindowsD3D12Viewport.cpp
--- a/Engine/Source/Runtime/D3D12RHI/Private/Windows/WindowsD3D12Viewport.cpp
+++ b/Engine/Source/Runtime/D3D12RHI/Private/Windows/WindowsD3D12Viewport.cpp
@@ -341,7 +341,7 @@ void FD3D12Viewport::ResizeInternal()
 	{
 		if (SwapChain1)
 		{
-			auto Lambda = [=]() -> FString
+			auto Lambda = [this, SwapChainFlags]() -> FString
 			{
 				return FString::Printf(TEXT("Num=%d, Size=(%d,%d), PF=%d, DXGIFormat=0x%x, Flags=0x%x"), NumBackBuffers, SizeX, SizeY, (int32)PixelFormat, (int32)GetRenderTargetFormat(PixelFormat), SwapChainFlags);
 			};
diff --git a/Engine/Source/Runtime/OpenGLDrv/Private/OpenGLQuery.cpp b/Engine/Source/Runtime/OpenGLDrv/Private/OpenGLQuery.cpp
--- a/Engine/Source/Runtime/OpenGLDrv/Private/OpenGLQuery.cpp
+++ b/Engine/Source/Runtime/OpenGLDrv/Private/OpenGLQuery.cpp
@@ -597,7 +597,7 @@ bool FOpenGLDynamicRHI::RHIGetRenderQueryResult(FRHIRenderQuery* QueryRHI, uint6
 			}
 			else
 			{
-				ALLOC_COMMAND_CL(RHICmdList, FRHICommandGLCommand)([=]() {GetRenderQueryResult_OnThisThread(ResourceCast(QueryRHI), false); GBatcher.Flush(*this, QueryRHI);  });
+				ALLOC_COMMAND_CL(RHICmdList, FRHICommandGLCommand)([this, QueryRHI]() {GetRenderQueryResult_OnThisThread(ResourceCast(QueryRHI), false); GBatcher.Flush(*this, QueryRHI);  });
 			}
 		}	
 	}
diff --git a/Engine/Source/Runtime/OpenGLDrv/Private/OpenGLTexture.cpp b/Engine/Source/Runtime/OpenGLDrv/Private/OpenGLTexture.cpp
--- a/Engine/Source/Runtime/OpenGLDrv/Private/OpenGLTexture.cpp
+++ b/Engine/Source/Runtime/OpenGLDrv/Private/OpenGLTexture.cpp
@@ -1475,7 +1475,7 @@ void FOpenGLDynamicRHI::RHIGenerateMips(FRHITexture* SurfaceRHI)
 {
 	if (FOpenGL::SupportsGenerateMipmap())
 	{
-		RunOnGLRenderContextThread([=]()
+		RunOnGLRenderContextThread([this, SurfaceRHI]()
 		{
 			VERIFY_GL_SCOPE();
 			GPUProfilingData.RegisterGPUWork(0);
@@ -2034,7 +2034,7 @@ void FOpenGLDynamicRHI::UnlockTexture2D_RenderThread(class FRHICommandListImmedi
 	}
 	else
 	{
-		auto GLCommand = [=]()
+		auto GLCommand = [this, Params, Texture, MipIndex, bLockWithinMiptail]()
 		{
 			uint32 DestStride;
 			uint64 LockedByteCount = ~0ULL;
@@ -2091,7 +2091,7 @@ void FOpenGLDynamicRHI::RHIUnlockTextureCubeFace_RenderThread(class FRHICommandL
 	}
 	else
 	{
-		auto GLCommand = [=]()
+		auto GLCommand = [this, Params, Texture, FaceIndex, ArrayIndex, MipIndex, bLockWithinMiptail]()
 		{
 			uint32 DestStride;
 			uint8* TexMem = (uint8*)this->RHILockTextureCubeFace(Texture, FaceIndex, ArrayIndex, MipIndex, RLM_WriteOnly, DestStride, bLockWithinMiptail);
@@ -2146,7 +2146,7 @@ void FOpenGLDynamicRHI::UnlockTexture2DArray_RenderThread(class FRHICommandListI
 	}
 	else
 	{
-		auto GLCommand = [=]()
+		auto GLCommand = [this, Params, Texture, ArrayIndex, MipIndex, bLockWithinMiptail]()
 		{
 			uint32 DestStride;
 			uint8* TexMem = (uint8*)this->RHILockTexture2DArray(Texture, ArrayIndex, MipIndex, Params.LockMode, DestStride, bLockWithinMiptail);
diff --git a/Engine/Source/Runtime/OpenGLDrv/Private/OpenGLViewport.cpp b/Engine/Source/Runtime/OpenGLDrv/Private/OpenGLViewport.cpp
--- a/Engine/Source/Runtime/OpenGLDrv/Private/OpenGLViewport.cpp
+++ b/Engine/Source/Runtime/OpenGLDrv/Private/OpenGLViewport.cpp
@@ -308,7 +308,7 @@ void FOpenGLViewport::Resize(uint32 InSizeX,uint32 InSizeY,bool bInIsFullscreen)
 			BackBuffer = new FOpenGLTexture(Desc);
 		}
 
-		RHICmdList.EnqueueLambda([=](FRHICommandListImmediate&)
+		RHICmdList.EnqueueLambda([this, InSizeX, InSizeY, bInIsFullscreen, bWasFullscreen](FRHICommandListImmediate&)
 		{
 			PlatformResizeGLContext(OpenGLRHI->PlatformDevice, OpenGLContext, InSizeX, InSizeY, bInIsFullscreen, bWasFullscreen, BackBuffer->Target, BackBuffer->GetResource());
 		});
