PR #14135:  Display ticked objects in ActorPalette viewport, when realtime toggle is enabled

#Synced-CL 49163689

diff --git a/Engine/Plugins/Experimental/ActorPalette/Source/ActorPalette/Private/ActorPaletteViewportClient.cpp b/Engine/Plugins/Experimental/ActorPalette/Source/ActorPalette/Private/ActorPaletteViewportClient.cpp
--- a/Engine/Plugins/Experimental/ActorPalette/Source/ActorPalette/Private/ActorPaletteViewportClient.cpp
+++ b/Engine/Plugins/Experimental/ActorPalette/Source/ActorPalette/Private/ActorPaletteViewportClient.cpp
@@ -135,6 +135,16 @@ FLinearColor FActorPaletteViewportClient::GetBackgroundColor() const
 	return MyBGColor;
 }
 
+void FActorPaletteViewportClient::Tick(float DeltaSeconds) 
+{
+	FEditorViewportClient::Tick(DeltaSeconds);
+
+	if ( this->IsRealtime() )
+	{
+		OwnedPreviewScene.GetWorld()->Tick(LEVELTICK_All, DeltaSeconds);
+	}
+}
+
 void FActorPaletteViewportClient::OpenWorldAsPalette(const FAssetData& InSourceWorldAsset)
 {
 	FScopedSlowTask SlowTask(0.0f, LOCTEXT("LoadingLevelAsPalette", "Loading Level into Actor Palette..."));
diff --git a/Engine/Plugins/Experimental/ActorPalette/Source/ActorPalette/Private/ActorPaletteViewportClient.h b/Engine/Plugins/Experimental/ActorPalette/Source/ActorPalette/Private/ActorPaletteViewportClient.h
--- a/Engine/Plugins/Experimental/ActorPalette/Source/ActorPalette/Private/ActorPaletteViewportClient.h
+++ b/Engine/Plugins/Experimental/ActorPalette/Source/ActorPalette/Private/ActorPaletteViewportClient.h
@@ -24,6 +24,7 @@ public:
 
 	// FEditorViewportClient interface
 	virtual FLinearColor GetBackgroundColor() const override;
+	void Tick(float DeltaSeconds) override;
 	// End of FEditorViewportClient interface
 
 	void OpenWorldAsPalette(const FAssetData& SourceWorldAsset);
